Capture Data plugin restructure

diff --git a/Engine/Plugins/VirtualProduction/CaptureData/CaptureData.uplugin b/Engine/Plugins/VirtualProduction/CaptureData/CaptureData.uplugin
new file mode 100644
--- /dev/null
+++ b/Engine/Plugins/VirtualProduction/CaptureData/CaptureData.uplugin
@@ -0,0 +1,39 @@
+{
+	"FileVersion": 1,
+	"Version": 1,
+	"VersionName": "1.0",
+	"FriendlyName": "Capture Data",
+	"Description": "Classes releated to captured data",
+	"Category": "Virtual Production",
+	"CreatedBy": "Epic Games, Inc.",
+	"CreatedByURL": "https://epicgames.com",
+	"DocsURL": "",
+	"MarketplaceURL": "",
+	"SupportURL": "",
+	"CanContainContent": false,
+	"Hidden": true,
+	"IsBetaVersion": false,
+	"Installed": false,
+	"Modules": [
+		{
+			"Name": "CaptureDataCore",
+			"Type": "Runtime",
+			"LoadingPhase": "Default"
+		},
+		{
+			"Name": "CaptureDataEditor",
+			"Type": "Editor",
+			"LoadingPhase": "Default"
+		}
+	],
+	"Plugins": [
+		{
+			"Name": "ImgMedia",
+			"Enabled": true
+		},
+		{
+			"Name": "CameraCalibrationCore",
+			"Enabled": true
+		}
+	]
+}
\ No newline at end of file
diff --git a/Engine/Plugins/VirtualProduction/CaptureData/Source/CaptureDataCore/CaptureDataCore.Build.cs b/Engine/Plugins/VirtualProduction/CaptureData/Source/CaptureDataCore/CaptureDataCore.Build.cs
new file mode 100644
--- /dev/null
+++ b/Engine/Plugins/VirtualProduction/CaptureData/Source/CaptureDataCore/CaptureDataCore.Build.cs
@@ -0,0 +1,27 @@
+// Copyright Epic Games, Inc. All Rights Reserved.
+
+using UnrealBuildTool;
+using System.IO;
+public class CaptureDataCore : ModuleRules
+{
+	public CaptureDataCore(ReadOnlyTargetRules Target) : base(Target)
+	{
+		PCHUsage = PCHUsageMode.UseExplicitOrSharedPCHs;
+
+		PublicDependencyModuleNames.AddRange(new string[]
+		{
+
+		});
+
+		PrivateDependencyModuleNames.AddRange(new string[]
+		{
+			"Core",
+			"CoreUObject",
+			"ImgMedia",
+			"CameraCalibrationCore",
+			"MeshDescription",
+			"StaticMeshDescription",
+			"Engine"
+		});
+	}
+}
diff --git a/Engine/Plugins/VirtualProduction/CaptureData/Source/CaptureDataCore/Private/CameraCalibration.cpp b/Engine/Plugins/VirtualProduction/CaptureData/Source/CaptureDataCore/Private/CameraCalibration.cpp
new file mode 100644
--- /dev/null
+++ b/Engine/Plugins/VirtualProduction/CaptureData/Source/CaptureDataCore/Private/CameraCalibration.cpp
@@ -0,0 +1,257 @@
+// Copyright Epic Games, Inc. All Rights Reserved.
+
+#include "CameraCalibration.h"
+#include "CaptureDataLog.h"
+#include "Models/SphericalLensModel.h"
+#include "OpenCVHelperLocal.h"
+#include "LensFile.h"
+#include "EditorFramework/AssetImportData.h"
+#include "AssetRegistry/AssetRegistryModule.h"
+#include "UObject/Package.h"
+#include "Misc/Paths.h"
+
+void UCameraCalibration::PostInitProperties()
+{
+	Super::PostInitProperties();
+
+#if WITH_EDITORONLY_DATA
+	if (!HasAnyFlags(RF_ClassDefaultObject))
+	{
+		AssetImportData = NewObject<UAssetImportData>(this, TEXT("AssetImportData"));
+	}
+#endif
+}
+
+void UCameraCalibration::PostLoad()
+{
+	Super::PostLoad();
+
+	// Back-compatability with older import where camera name was not recorded.
+	// These always have 2 cameras, the first being rgb the seconds being depth.
+	// Distinguish between iphone and HMC import by looking at relative size of RGB and depth images.
+	// The RGB camera for the iphone case we'll call "iPhone"
+	// The RGB camera for the HMC case will be "bot"
+
+	if (CameraCalibrations.Num() == 2)
+	{
+		if (CameraCalibrations[0].Name.IsEmpty() && CameraCalibrations[1].Name.IsEmpty())
+		{
+			if (CameraCalibrations[0].LensFile->LensInfo.ImageDimensions.X == CameraCalibrations[1].LensFile->LensInfo.ImageDimensions.X * 2)
+			{
+				CameraCalibrations[0].Name = TEXT("iPhone");
+			}
+			else
+			{
+				CameraCalibrations[0].Name = TEXT("bot");
+			}
+
+			CameraCalibrations[1].Name = TEXT("Depth");
+		}
+	}
+}
+
+void UCameraCalibration::GetAssetRegistryTags(TArray<FAssetRegistryTag>& OutTags) const
+{
+#if WITH_EDITORONLY_DATA
+	if (AssetImportData != nullptr)
+	{
+		OutTags.Emplace(SourceFileTagName(), AssetImportData->GetSourceData().ToJson(), FAssetRegistryTag::TT_Hidden);
+	}
+#endif
+}
+
+bool UCameraCalibration::ConvertToTrackerNodeCameraModels(TArray<FCameraCalibration>& OutCalibrations,
+																   TArray<TPair<FString, FString>>& OutStereoReconstructionPairs) const
+{
+	OutCalibrations.SetNum(CameraCalibrations.Num());
+	OutStereoReconstructionPairs.SetNum(StereoPairs.Num());
+
+	for (int32 Pair = 0; Pair < OutStereoReconstructionPairs.Num(); Pair++)
+	{
+		OutStereoReconstructionPairs[Pair] = TPair<FString, FString>(FString::FromInt(StereoPairs[Pair].CameraIndex1), FString::FromInt(StereoPairs[Pair].CameraIndex2));
+	}
+
+
+	for (int32 Cam = 0; Cam < CameraCalibrations.Num(); Cam++)
+	{
+		FCameraCalibration CurCalib;
+		CurCalib.Name = CameraCalibrations[Cam].Name;
+		CurCalib.ImageSizeX = CameraCalibrations[Cam].LensFile->LensInfo.ImageDimensions.X;
+		CurCalib.ImageSizeY = CameraCalibrations[Cam].LensFile->LensInfo.ImageDimensions.Y;
+
+		if (CameraCalibrations[Cam].LensFile->LensInfo.LensModel != USphericalLensModel::StaticClass())
+		{
+			UE_LOG(LogCaptureDataCore, Warning, TEXT("Camera calibration does not contain a SphericalLensModel lens distortion."));
+			return false;
+		}
+
+		const FDistortionTable& DistortionTable = CameraCalibrations[Cam].LensFile->DistortionTable;
+		FDistortionInfo DistortionData;
+		bool bGotPoint = DistortionTable.GetPoint(0.0f, 0.0f, DistortionData);
+
+		if (!bGotPoint)
+		{
+			UE_LOG(LogCaptureDataCore, Warning, TEXT("Camera calibration does not contain a valid lens distortion."));
+			return false;
+		}
+
+		check(DistortionData.Parameters.Num() == 5);
+		CurCalib.K1 = DistortionData.Parameters[0];
+		CurCalib.K2 = DistortionData.Parameters[1];
+		CurCalib.K3 = DistortionData.Parameters[2]; // parameters are stored K1 K2 K3 P1 P2 rather than OpenCV order of K1 K2 P1 P2 K3
+		CurCalib.P1 = DistortionData.Parameters[3];
+		CurCalib.P2 = DistortionData.Parameters[4];
+		CurCalib.K4 = 0.0;
+		CurCalib.K5 = 0.0;
+		CurCalib.K6 = 0.0;
+
+		const FFocalLengthTable& FocalLengthTable = CameraCalibrations[Cam].LensFile->FocalLengthTable;
+		FFocalLengthInfo FocalLengthData;
+		bGotPoint = FocalLengthTable.GetPoint(0.0f, 0.0f, FocalLengthData);
+
+		if (!bGotPoint)
+		{
+			UE_LOG(LogCaptureDataCore, Warning, TEXT("Camera calibration does not contain a valid focal length."));
+			return false;
+		}
+		const FVector2D& FxFy = FocalLengthData.FxFy;
+		CurCalib.FX = FxFy.X * CurCalib.ImageSizeX;
+		CurCalib.FY = FxFy.Y * CurCalib.ImageSizeY;
+
+		const FImageCenterTable& ImageCenterTable = CameraCalibrations[Cam].LensFile->ImageCenterTable;
+		FImageCenterInfo ImageCenterData;
+		bGotPoint = ImageCenterTable.GetPoint(0.0f, 0.0f, ImageCenterData);
+
+		if (!bGotPoint)
+		{
+			UE_LOG(LogCaptureDataCore, Warning, TEXT("Camera calibration does not contain a valid image center."));
+			return false;
+		}
+		const FVector2D& PrincipalPoint = ImageCenterData.PrincipalPoint;
+		CurCalib.CX = PrincipalPoint.X * CurCalib.ImageSizeX;
+		CurCalib.CY = PrincipalPoint.Y * CurCalib.ImageSizeY;
+
+		const FNodalOffsetTable& NodalOffsetTable = CameraCalibrations[Cam].LensFile->NodalOffsetTable;
+		FNodalPointOffset NodalOffsetData;
+		bGotPoint = NodalOffsetTable.GetPoint(0.0f, 0.0f, NodalOffsetData);
+
+		if (!bGotPoint)
+		{
+			UE_LOG(LogCaptureDataCore, Warning, TEXT("Camera calibration does not contain a valid nodal offset."));
+			return false;
+		}
+
+		FTransform Transform;
+		Transform.SetLocation(NodalOffsetData.LocationOffset);
+		Transform.SetRotation(NodalOffsetData.RotationOffset);
+		FOpenCVHelperLocal::ConvertUnrealToOpenCV(Transform);
+		CurCalib.Transform = Transform.ToMatrixWithScale();
+
+		OutCalibrations[Cam] = CurCalib;
+	}
+
+	return true;
+}
+
+bool UCameraCalibration::ConvertFromTrackerNodeCameraModels(TArray<FCameraCalibration>& InCalibrations)
+{
+	check(InCalibrations.Num() == 2 || InCalibrations.Num() == 3); // only works for stereo hmc, or for the case of one RGB stream and one depth stream
+
+	for (int32 Index = 0; Index < InCalibrations.Num(); ++Index)
+	{
+		const FCameraCalibration& Calibration = InCalibrations[Index];
+
+		bool bIsDepthCamera = (Index == InCalibrations.Num() - 1);
+
+		const FString ObjectName = bIsDepthCamera ?
+			FString::Printf(TEXT("%s_Depth_LensFile"), *GetName()) :
+			FString::Printf(TEXT("%s_%s_RGB_LensFile"), *GetName(), *Calibration.Name);
+		FString ParentPath = FString::Printf(TEXT("%s/../%s"), *GetPackage()->GetPathName(), *ObjectName);
+		FPaths::CollapseRelativeDirectories(ParentPath);
+		UObject* Parent = CreatePackage(*ParentPath);
+
+		FExtendedLensFile CameraCalibration;
+		CameraCalibration.Name = Calibration.Name;
+		CameraCalibration.IsDepthCamera = bIsDepthCamera;
+		CameraCalibration.LensFile = NewObject<ULensFile>(Parent, ULensFile::StaticClass(), *ObjectName, GetFlags());
+
+		// These a for a non-FIZ camera.
+		const float Focus = 0.0f;
+		const float Zoom = 0.0f;
+
+		// LensInfo
+		CameraCalibration.LensFile->LensInfo.LensModel = USphericalLensModel::StaticClass();
+		CameraCalibration.LensFile->LensInfo.LensModelName = FString::Printf(TEXT("Lens"));
+		// lens serial number is not needed
+		
+		// leave sensor dimensions with default values and de-normalize using VideoDimensions or DepthDimensions
+		CameraCalibration.LensFile->LensInfo.ImageDimensions = FIntPoint(Calibration.ImageSizeX, Calibration.ImageSizeY);
+
+		// FocalLengthInfo
+		FFocalLengthInfo FocalLengthInfo;
+		FocalLengthInfo.FxFy = FVector2D(Calibration.FX / Calibration.ImageSizeX, Calibration.FY / Calibration.ImageSizeY);
+
+		// DistortionInfo
+		FDistortionInfo DistortionInfo;
+		FSphericalDistortionParameters SphericalParameters;
+
+		SphericalParameters.K1 = Calibration.K1;
+		SphericalParameters.K2 = Calibration.K2;
+		SphericalParameters.P1 = Calibration.P1;
+		SphericalParameters.P2 = Calibration.P2;
+		SphericalParameters.K3 = Calibration.K3;
+
+		USphericalLensModel::StaticClass()->GetDefaultObject<ULensModel>()->ToArray(
+			SphericalParameters,
+			DistortionInfo.Parameters
+		);
+
+		// ImageCenterInfo
+		FImageCenterInfo ImageCenterInfo;
+		ImageCenterInfo.PrincipalPoint = FVector2D(Calibration.CX / Calibration.ImageSizeX, Calibration.CY / Calibration.ImageSizeY);
+
+		// NodalOffset
+		FNodalPointOffset NodalPointOffset;
+		FTransform Transform;
+		Transform.SetFromMatrix(Calibration.Transform);
+		FOpenCVHelperLocal::ConvertOpenCVToUnreal(Transform);
+		NodalPointOffset.LocationOffset = Transform.GetLocation();
+		NodalPointOffset.RotationOffset = Transform.GetRotation();
+
+		CameraCalibration.LensFile->AddDistortionPoint(Focus, Zoom, DistortionInfo, FocalLengthInfo);
+		CameraCalibration.LensFile->AddImageCenterPoint(Focus, Zoom, ImageCenterInfo);
+		CameraCalibration.LensFile->AddNodalOffsetPoint(Focus, Zoom, NodalPointOffset);
+
+		CameraCalibration.LensFile->MarkPackageDirty();
+		FAssetRegistryModule::AssetCreated(CameraCalibration.LensFile);
+
+		CameraCalibrations.Add(CameraCalibration);
+	}
+
+	if (InCalibrations.Num() == 3)
+	{
+		// stereo HMC so make the stereo pair
+		FStereoPair CameraPair;
+		CameraPair.CameraIndex1 = 0;
+		CameraPair.CameraIndex2 = 1;
+		StereoPairs.Add(CameraPair);
+	}
+
+	return true;
+}
+
+int32 UCameraCalibration::GetIndexByCameraName(const FString& InName) const
+{
+	for (int32 Calibration = 0; Calibration < CameraCalibrations.Num(); ++Calibration)
+	{
+		if (CameraCalibrations[Calibration].Name == InName)
+		{
+			return Calibration;
+		}
+	}
+
+	UE_LOG(LogCaptureDataCore, Warning, TEXT("Specified camera name not valid"));
+
+	return -1;
+}
+
diff --git a/Engine/Plugins/VirtualProduction/CaptureData/Source/CaptureDataCore/Private/CaptureData.cpp b/Engine/Plugins/VirtualProduction/CaptureData/Source/CaptureDataCore/Private/CaptureData.cpp
new file mode 100644
--- /dev/null
+++ b/Engine/Plugins/VirtualProduction/CaptureData/Source/CaptureDataCore/Private/CaptureData.cpp
@@ -0,0 +1,672 @@
+// Copyright Epic Games, Inc. All Rights Reserved.
+
+#include "CaptureData.h"
+
+#include "CaptureDataLog.h"
+#include "CameraCalibration.h"
+
+#include "ImageSequenceUtils.h"
+
+#include "Components/SkeletalMeshComponent.h"
+#include "Components/StaticMeshComponent.h"
+#include "Rendering/SkeletalMeshModel.h"
+#include "Rendering/SkeletalMeshLODModel.h"
+#include "StaticMeshAttributes.h"
+#include "MeshAttributes.h"
+#include "Engine/StaticMesh.h"
+#include "Misc/TransactionObjectEvent.h"
+#include "Sound/SoundWave.h"
+#include "ImgMediaSource.h"
+#include "MediaTexture.h"
+#include "Algo/AllOf.h"
+#include "Engine/SkeletalMesh.h"
+
+#if WITH_EDITOR
+#include "IDirectoryWatcher.h"
+#include "DirectoryWatcherModule.h"
+#endif
+
+#include UE_INLINE_GENERATED_CPP_BY_NAME(CaptureData)
+
+/////////////////////////////////////////////////////
+// UCaptureData
+
+#if WITH_EDITOR
+
+void UCaptureData::PostEditChangeProperty(struct FPropertyChangedEvent& InPropertyChangedEvent)
+{
+	Super::PostEditChangeProperty(InPropertyChangedEvent);
+	NotifyInternalsChanged();
+}
+
+void UCaptureData::PostTransacted(const FTransactionObjectEvent& InTransactionEvent)
+{
+	Super::PostTransacted(InTransactionEvent);
+	NotifyInternalsChanged();
+}
+
+#endif
+
+void UCaptureData::NotifyInternalsChanged()
+{
+	OnCaptureDataInternalsChangedDelegate.Broadcast();
+}
+
+/////////////////////////////////////////////////////
+// UMeshCaptureData
+
+bool UMeshCaptureData::IsInitialized() const
+{
+	// We need access to the source data when fitting and that is not available during runtime
+#if WITH_EDITOR
+	return TargetMesh != nullptr && (TargetMesh->IsA<UStaticMesh>() || TargetMesh->IsA<USkeletalMesh>());
+#else
+	return false;
+#endif
+}
+
+void UMeshCaptureData::GetDataForConforming(const FTransform& InTransform, TArray<float>& OutVertices, TArray<int32>& OutTriangles) const
+{
+#if WITH_EDITOR
+	if (USkeletalMesh* TargetSkeletalMesh = Cast<USkeletalMesh>(TargetMesh))
+	{
+
+		FSkeletalMeshModel* ImportedModel = TargetSkeletalMesh->GetImportedModel();
+		FSkeletalMeshLODModel& LODModel = ImportedModel->LODModels[0];
+
+		OutVertices.Reset(LODModel.NumVertices * 3);
+
+		for (FSkelMeshSection& Section : LODModel.Sections)
+		{
+			const int32 NumVertices = Section.GetNumVertices();
+			for (int32 VertexIndex = 0; VertexIndex < NumVertices; ++VertexIndex)
+			{
+				FSoftSkinVertex& OriginalVertex = Section.SoftVertices[VertexIndex];
+				FVector TransformedVertex = InTransform.TransformPosition(FVector{ OriginalVertex.Position });
+				OutVertices.Add(TransformedVertex.Y);
+				OutVertices.Add(-TransformedVertex.Z);
+				OutVertices.Add(TransformedVertex.X);
+			}
+		}
+
+		OutTriangles.Reset(LODModel.IndexBuffer.Num());
+		for (uint32 Index : LODModel.IndexBuffer)
+		{
+			OutTriangles.Add(Index);
+		}
+	}
+	else if (UStaticMesh* TargetStaticMesh = Cast<UStaticMesh>(TargetMesh))
+	{
+		FMeshDescription* MeshDescription = TargetStaticMesh->GetMeshDescription(0);
+		check(MeshDescription);
+		FStaticMeshAttributes Attributes(*MeshDescription);
+
+		TVertexAttributesRef<FVector3f> OriginalMeshVerts = Attributes.GetVertexPositions();
+		TTriangleAttributesRef<TArrayView<FVertexID>>  OriginalMeshIndices = Attributes.GetTriangleVertexIndices();
+
+		OutVertices.Reset(OriginalMeshVerts.GetNumElements());
+
+		for (int32 RenderCtr = 0; RenderCtr < OriginalMeshVerts.GetNumElements(); ++RenderCtr)
+		{
+			// map the mesh vertices (in UE coordinate system) to OpenCV coordinate system
+			FVector3f OriginalVertex = OriginalMeshVerts.Get(RenderCtr);
+			FVector TransformedVertex = InTransform.TransformPosition(FVector{ OriginalVertex });
+			OutVertices.Add(TransformedVertex.Y);
+			OutVertices.Add(-TransformedVertex.Z);
+			OutVertices.Add(TransformedVertex.X);
+		}
+
+		OutTriangles.Reset(OriginalMeshIndices.GetNumElements());
+		const TArrayView<FVertexID>& RawIndArray = OriginalMeshIndices.GetRawArray();
+		for (const auto& Index : RawIndArray)
+		{
+			OutTriangles.Add(Index.GetValue());
+		}
+	}
+	else
+	{
+		// This is an error state so log it accordingly
+		if (TargetMesh != nullptr)
+		{
+			UE_LOG(LogCaptureDataCore, Error, TEXT("Failed to get data for conforming as TargetMesh is a '%s' but should be a UStaticMesh or USkeletalMesh"), *TargetMesh->GetClass()->GetName());
+		}
+		else
+		{
+			UE_LOG(LogCaptureDataCore, Error, TEXT("Failed to get data for conforming as TargetMesh is invalid"));
+		}
+	}
+#endif
+}
+
+/////////////////////////////////////////////////////
+// UFootageCaptureData
+
+UFootageCaptureData::FPathAssociation::FPathAssociation(const FString& InPathOnDisk, const FString& InAssetPath) :
+	PathOnDisk(InPathOnDisk),
+	AssetPath(InAssetPath)
+{
+}
+
+TOptional<FFootageCaptureMetadata::FIosDeviceVersion> FFootageCaptureMetadata::ParseIosDeviceVersion(const FString& Prefix, const FString& ModelName)
+{
+	FString CombinedModelNumbers;
+	if (ModelName.Split(Prefix, nullptr, &CombinedModelNumbers)) // Combined model numbers, e.g. "1,2"
+	{
+		TArray<FString> SeparatedModelNumbers;
+		if (CombinedModelNumbers.ParseIntoArray(SeparatedModelNumbers, TEXT(",")) == 2) // Separated model numbers, e.g. ["1", "2"]
+		{
+			return FIosDeviceVersion(FCString::Atoi(*SeparatedModelNumbers[0]), FCString::Atoi(*SeparatedModelNumbers[1]));
+		}
+	}
+
+	return TOptional<FIosDeviceVersion>();
+}
+
+EFootageDeviceClass FFootageCaptureMetadata::IPhoneDeviceClass(TOptional<FIosDeviceVersion> IosDeviceVersion)
+{
+	if (IosDeviceVersion.IsSet())
+	{
+		const FIosDeviceVersion Version = IosDeviceVersion.GetValue();
+		const uint16 Major = Version.Major;
+		const uint16 Minor = Version.Minor;
+		if (Major < 12) // Before iPhone 11
+		{
+			return EFootageDeviceClass::iPhone11OrEarlier;
+		}
+		else if (Major == 12) // iPhone 11
+		{
+			const TArray<uint16> IPhone11MinorVersions = { 1, 3, 5 };
+			if (IPhone11MinorVersions.Contains(Minor))
+			{
+				return EFootageDeviceClass::iPhone11OrEarlier;
+			}
+		}
+		else if (Major == 13) // iPhone 12
+		{
+			return EFootageDeviceClass::iPhone12;
+		}
+		else if (Major == 14) // iPhone 13 or non pro iPhone 14 models
+		{
+			const TArray<uint16> IPhone13MinorVersions = { 2, 3, 4, 5 };
+			const TArray<uint16> IPhone14MinorVersions = { 7, 8 };
+			if (IPhone13MinorVersions.Contains(Minor))
+			{
+				return EFootageDeviceClass::iPhone13;
+			}
+			else if (IPhone14MinorVersions.Contains(Minor))
+			{
+				return EFootageDeviceClass::iPhone14OrLater;
+			}
+		}
+		else if (Major >= 15) // iPhone 14 Pro Models or later
+		{
+			return  EFootageDeviceClass::iPhone14OrLater;
+		}
+	}
+
+	return EFootageDeviceClass::OtheriOSDevice;
+}
+
+void FFootageCaptureMetadata::SetDeviceClass(const FString& InDeviceModel)
+{
+	const FString IPhone = TEXT("iPhone");
+	const FString IPad = TEXT("iPad");
+	const FString StereoHMC = TEXT("StereoHMC");
+
+	if (InDeviceModel.StartsWith(IPhone))
+	{
+		DeviceClass = IPhoneDeviceClass(ParseIosDeviceVersion(IPhone, InDeviceModel));
+	}
+	else if (InDeviceModel.StartsWith(IPad))
+	{
+		DeviceClass = EFootageDeviceClass::OtheriOSDevice;
+	}
+	else if (InDeviceModel == StereoHMC)
+	{
+		DeviceClass = EFootageDeviceClass::StereoHMC;
+	}
+	else
+	{
+		DeviceClass = EFootageDeviceClass::Unspecified;
+	}
+}
+
+#define CHECK_AND_PRINT_RET(InCondition, InFormat, ...) if (!(InCondition)) { return MakeError(FString::Printf(InFormat, ##__VA_ARGS__)); }
+#define CHECK_AND_RET_ERROR(InFunction) if (UFootageCaptureData::FVerifyResult Result = InFunction; Result.HasError()) { return MoveTemp(Result); }
+
+bool UFootageCaptureData::IsInitialized() const
+{
+	return !VerifyData().HasError();
+}
+
+UFootageCaptureData::FVerifyResult UFootageCaptureData::VerifyData() const
+{
+	CHECK_AND_RET_ERROR(ViewsContainsValidData());
+	CHECK_AND_RET_ERROR(MetadataContainsValidData());
+	CHECK_AND_RET_ERROR(CalibrationContainsValidData());
+
+	return MakeValue();
+}
+
+TArray<UFootageCaptureData::FPathAssociation> UFootageCaptureData::CheckImageSequencePaths() const
+{
+	TArray<FPathAssociation> InvalidImageSequences;
+
+	for (const FFootageCaptureView& View : Views)
+	{
+		const UImgMediaSource* ImageSequence = View.ImageSequence;
+
+		if (ImageSequence && !FPaths::DirectoryExists(ImageSequence->GetFullPath()))
+		{
+			InvalidImageSequences.Emplace(ImageSequence->GetFullPath(), ImageSequence->GetPathName());
+		}
+	}
+
+	return InvalidImageSequences;
+}
+
+UFootageCaptureData::FVerifyResult UFootageCaptureData::ViewsContainsValidData() const
+{
+	CHECK_AND_PRINT_RET(!Views.IsEmpty(), TEXT("Capture Data doesn't contain views"));
+
+	for (const FFootageCaptureView& View : Views)
+	{
+		int32 NumImageFrames = 0;
+		FIntVector2 ImageDimensions;
+
+		CHECK_AND_PRINT_RET(FImageSequenceUtils::GetImageSequenceInfo(View.ImageSequence, ImageDimensions, NumImageFrames), TEXT("Image Sequence asset is invalid"));
+
+		int32 NumDepthFrames = 0;
+		FIntVector2 DepthDimensions;
+
+		CHECK_AND_PRINT_RET(FImageSequenceUtils::GetImageSequenceInfo(View.DepthSequence, DepthDimensions, NumImageFrames), TEXT("Depth Sequence asset is invalid"));
+
+		CHECK_AND_PRINT_RET(View.ImageSequence->FrameRateOverride.IsValid(), TEXT("Image Sequence asset contains invalid frame rate"));
+		CHECK_AND_PRINT_RET(View.DepthSequence->FrameRateOverride.IsValid(), TEXT("Depth Sequence asset contains invalid frame rate"));
+		CHECK_AND_PRINT_RET(View.ImageSequence->FrameRateOverride == View.DepthSequence->FrameRateOverride, TEXT("Image Sequence asset and Depth Sequence asset contain different frame rates"));
+	}
+
+	return MakeValue();
+}
+
+UFootageCaptureData::FVerifyResult UFootageCaptureData::MetadataContainsValidData() const
+{
+	CHECK_AND_PRINT_RET(Metadata.Height > 0 && Metadata.Width > 0, TEXT("Invalid values for the image resolution"));
+	CHECK_AND_PRINT_RET(!FMath::IsNearlyZero(Metadata.FrameRate), TEXT("Frame rate can't be set to 0"));
+
+	return MakeValue();
+}
+
+UFootageCaptureData::FVerifyResult UFootageCaptureData::CalibrationContainsValidData() const
+{
+	CHECK_AND_PRINT_RET(!CameraCalibrations.IsEmpty(), TEXT("Calibration assets are empty"));
+	for (UCameraCalibration* CameraCalibration : CameraCalibrations)
+	{
+		CHECK_AND_PRINT_RET(!CameraCalibration->CameraCalibrations.IsEmpty(), TEXT("Calibration asset not configured"));
+	}
+	return MakeValue();
+}
+
+#undef CHECK_AND_PRINT_RET
+#undef CHECK_AND_RET
+
+static TRange<FFrameNumber> GetFrameRange(const FFrameRate& InTargetRate, const FTimecode& InMediaTimecode, const FFrameRate& InMediaTimecodeRate, bool bInMediaStartFrameIsZero, const FFrameRate& InMediaRate, const FFrameNumber& InMediaDuration)
+{
+	const FFrameNumber MediaStartFrame = bInMediaStartFrameIsZero ? 0 : InMediaTimecode.ToFrameNumber(InMediaTimecodeRate);
+	const FFrameTime TargetStartFrameTime = FFrameRate::TransformTime(MediaStartFrame, InMediaTimecodeRate, InTargetRate);
+	const FFrameTime TargetDurationTime = FFrameRate::TransformTime(InMediaDuration, InMediaRate, InTargetRate);
+	const FFrameTime TargetEndFrameTime = TargetStartFrameTime + TargetDurationTime;
+	// When converting from frame time to frame number, deal with sub frame times by taking minimum frame range.
+	TRange<FFrameNumber> TargetFrameRange(TargetStartFrameTime.CeilToFrame(), TargetEndFrameTime.FloorToFrame());
+	return TargetFrameRange;
+}
+
+static TRange<FFrameNumber> GetFrameRange(const FFrameRate& InTargetRate, UImgMediaSource* InMedia, const FTimecode& InMediaTimecode, const FFrameRate& InMediaTimecodeRate, bool bInMediaStartFrameIsZero)
+{
+	int32 Duration = 0;
+	FIntVector2 ImageDimensions;
+	const bool bImageOK = FImageSequenceUtils::GetImageSequenceInfo(InMedia, ImageDimensions, Duration);
+	ensure(bImageOK);
+
+	return GetFrameRange(InTargetRate, InMediaTimecode, InMediaTimecodeRate, bInMediaStartFrameIsZero, InMedia->FrameRateOverride, Duration);
+}
+
+static TRange<FFrameNumber> GetFrameRange(const FFrameRate& InTargetRate, USoundWave* InMedia, const FTimecode& InMediaTimecode, const FFrameRate& InMediaTimecodeRate, bool bInMediaStartFrameIsZero)
+{
+	const int32 Duration = InMedia->GetDuration() * InMediaTimecodeRate.AsDecimal();
+
+	return GetFrameRange(InTargetRate, InMediaTimecode, InMediaTimecodeRate, bInMediaStartFrameIsZero, InMediaTimecodeRate, Duration);
+}
+
+FIntPoint UFootageCaptureData::GetFootageColorResolution() const
+{
+	FIntPoint Resolution(ForceInit);
+
+	if (!CameraCalibrations.IsEmpty())
+	{
+		TArray<FCameraCalibration> Calibrations;
+		TArray<TPair<FString, FString>> StereoPairs;
+		CameraCalibrations[0]->ConvertToTrackerNodeCameraModels(Calibrations, StereoPairs);
+
+		Resolution = FIntPoint(Calibrations[0].ImageSizeX, Calibrations[0].ImageSizeY);
+	}
+
+	return Resolution;
+}
+
+void UFootageCaptureData::GetFrameRanges(const FFrameRate& InTargetRate, ETimecodeAlignment InTimecodeAlignment, bool bInIncludeAudio, TMap<TWeakObjectPtr<UObject>, TRange<FFrameNumber>>& OutMediaFrameRanges, TRange<FFrameNumber>& OutProcessingFrameRange, TRange<FFrameNumber>& OutMaximumFrameRange) const
+{
+	OutMediaFrameRanges.Reset();
+	OutProcessingFrameRange = TRange<FFrameNumber>(0, 0);
+	OutMaximumFrameRange = TRange<FFrameNumber>(0, 0);
+
+	for (int32 Index = 0; Index < Views.Num(); ++Index)
+	{
+		OutMediaFrameRanges.Add(Views[Index].ImageSequence, GetFrameRange(InTargetRate, Views[Index].ImageSequence, GetEffectiveImageTimecode(Index), GetEffectiveImageTimecodeRate(Index), InTimecodeAlignment == ETimecodeAlignment::None));
+		OutMediaFrameRanges.Add(Views[Index].DepthSequence, GetFrameRange(InTargetRate, Views[Index].DepthSequence, GetEffectiveDepthTimecode(Index), GetEffectiveDepthTimecodeRate(Index), InTimecodeAlignment == ETimecodeAlignment::None));
+	}
+
+	if (bInIncludeAudio)
+	{
+		for (TObjectPtr<class USoundWave> Audio : Audios)
+		{
+			OutMediaFrameRanges.Add(Audio, GetFrameRange(InTargetRate, Audio, GetEffectiveAudioTimecode(), GetEffectiveAudioTimecodeRate(), InTimecodeAlignment == ETimecodeAlignment::None));
+		}
+	}
+
+	if (InTimecodeAlignment == ETimecodeAlignment::Relative)
+	{
+		FFrameNumber LowestStartFrame;
+
+		bool bFirstPass = true;
+		for (const TPair<TWeakObjectPtr<UObject>, TRange<FFrameNumber>>& MediaFrameRangePair : OutMediaFrameRanges)
+		{
+			if (bFirstPass || MediaFrameRangePair.Value.GetLowerBoundValue() < LowestStartFrame)
+			{
+				LowestStartFrame = MediaFrameRangePair.Value.GetLowerBoundValue();
+			}
+
+			bFirstPass = false;
+		}
+
+		for (TPair<TWeakObjectPtr<UObject>, TRange<FFrameNumber>>& MediaFrameRangePair : OutMediaFrameRanges)
+		{
+			MediaFrameRangePair.Value.SetLowerBoundValue(MediaFrameRangePair.Value.GetLowerBoundValue() - LowestStartFrame);
+			MediaFrameRangePair.Value.SetUpperBoundValue(MediaFrameRangePair.Value.GetUpperBoundValue() - LowestStartFrame);
+		}
+	}
+
+	bool bFirstPass = true;
+	for (const TPair<TWeakObjectPtr<UObject>, TRange<FFrameNumber>>& MediaFrameRangePair : OutMediaFrameRanges)
+	{
+		if (MediaFrameRangePair.Key.IsValid() &&
+			(MediaFrameRangePair.Key->IsA(UImgMediaSource::StaticClass()) || MediaFrameRangePair.Key->IsA(USoundWave::StaticClass())))
+		{
+			if (bFirstPass || MediaFrameRangePair.Value.GetLowerBoundValue() > OutProcessingFrameRange.GetLowerBoundValue())
+			{
+				OutProcessingFrameRange.SetLowerBoundValue(MediaFrameRangePair.Value.GetLowerBoundValue());
+			}
+
+			if (bFirstPass || MediaFrameRangePair.Value.GetUpperBoundValue() < OutProcessingFrameRange.GetUpperBoundValue())
+			{
+				OutProcessingFrameRange.SetUpperBoundValue(MediaFrameRangePair.Value.GetUpperBoundValue());
+			}
+
+			if (bFirstPass || MediaFrameRangePair.Value.GetLowerBoundValue() < OutMaximumFrameRange.GetLowerBoundValue())
+			{
+				OutMaximumFrameRange.SetLowerBoundValue(MediaFrameRangePair.Value.GetLowerBoundValue());
+			}
+
+			if (bFirstPass || MediaFrameRangePair.Value.GetUpperBoundValue() > OutMaximumFrameRange.GetUpperBoundValue())
+			{
+				OutMaximumFrameRange.SetUpperBoundValue(MediaFrameRangePair.Value.GetUpperBoundValue());
+			}
+
+			bFirstPass = false;
+		}
+	}
+
+	if (OutProcessingFrameRange.GetUpperBoundValue() <= OutProcessingFrameRange.GetLowerBoundValue())
+	{
+		OutProcessingFrameRange = TRange<FFrameNumber>(0, 0);
+	}
+}
+
+TRange<FFrameNumber> UFootageCaptureData::GetAudioFrameRange(const FFrameRate& InTargetRate, ETimecodeAlignment InTimecodeAlignment, USoundWave* InMedia, const FTimecode& InMediaTimecode, const FFrameRate& InMediaTimecodeRate)
+{
+	bool bInMediaStartFrameIsZero = InTimecodeAlignment == ETimecodeAlignment::None;
+	TRange<FFrameNumber> AudioFrameRange = GetFrameRange(InTargetRate, InMedia, InMediaTimecode, InMediaTimecodeRate, bInMediaStartFrameIsZero);
+
+	if (InTimecodeAlignment == ETimecodeAlignment::Relative)
+	{
+		FFrameNumber StartFrame = AudioFrameRange.GetLowerBoundValue();
+		AudioFrameRange.SetLowerBoundValue(AudioFrameRange.GetLowerBoundValue() - StartFrame);
+		AudioFrameRange.SetUpperBoundValue(AudioFrameRange.GetUpperBoundValue() - StartFrame);
+	}
+
+	return AudioFrameRange;
+}
+
+FTimecode UFootageCaptureData::GetEffectiveImageTimecode(int32 InView) const
+{
+	check(InView < Views.Num());
+
+	FTimecode Timecode;
+
+	if (Views[InView].bImageTimecodePresent)
+	{
+		Timecode = Views[InView].ImageTimecode;
+	}
+	else if (Views[InView].bDepthTimecodePresent)
+	{
+		Timecode = Views[InView].DepthTimecode;
+	}
+	else
+	{
+		FFrameRate FrameRate;
+		GetDefaultTimecodeInfo(Timecode, FrameRate);
+	}
+
+	return Timecode;
+}
+
+FFrameRate UFootageCaptureData::GetEffectiveImageTimecodeRate(int32 InView) const
+{
+	check(InView < Views.Num());
+
+	FFrameRate FrameRate;
+
+	if (Views[InView].bImageTimecodePresent)
+	{
+		FrameRate = Views[InView].ImageTimecodeRate;
+	}
+	else if (Views[InView].bDepthTimecodePresent)
+	{
+		FrameRate = Views[InView].DepthTimecodeRate;
+	}
+	else
+	{
+		FTimecode Timecode;
+		GetDefaultTimecodeInfo(Timecode, FrameRate);
+	}
+
+	return FrameRate;
+}
+
+FTimecode UFootageCaptureData::GetEffectiveDepthTimecode(int32 InView) const
+{
+	check(InView < Views.Num());
+
+	FTimecode Timecode;
+
+	if (Views[InView].bDepthTimecodePresent)
+	{
+		Timecode = Views[InView].DepthTimecode;
+	}
+	else if (Views[InView].bImageTimecodePresent)
+	{
+		Timecode = Views[InView].ImageTimecode;
+	}
+	else
+	{
+		FFrameRate FrameRate;
+		GetDefaultTimecodeInfo(Timecode, FrameRate);
+	}
+
+	return Timecode;
+}
+
+FFrameRate UFootageCaptureData::GetEffectiveDepthTimecodeRate(int32 InView) const
+{
+	check(InView < Views.Num());
+
+	FFrameRate FrameRate;
+
+	if (Views[InView].bDepthTimecodePresent)
+	{
+		FrameRate = Views[InView].DepthTimecodeRate;
+	}
+	else if (Views[InView].bImageTimecodePresent)
+	{
+		FrameRate = Views[InView].ImageTimecodeRate;
+	}
+	else
+	{
+		FTimecode Timecode;
+		GetDefaultTimecodeInfo(Timecode, FrameRate);
+	}
+
+	return FrameRate;
+}
+
+FTimecode UFootageCaptureData::GetEffectiveAudioTimecode() const
+{
+	FTimecode Timecode;
+
+	if (bAudioTimecodePresent)
+	{
+		Timecode = AudioTimecode;
+	}
+	else
+	{
+		FFrameRate FrameRate;
+		GetDefaultTimecodeInfo(Timecode, FrameRate);
+	}
+
+	return Timecode;
+}
+
+FFrameRate UFootageCaptureData::GetEffectiveAudioTimecodeRate() const
+{
+	FFrameRate FrameRate;
+
+	if (bAudioTimecodePresent)
+	{
+		FrameRate = AudioTimecodeRate;
+	}
+	else
+	{
+		FTimecode Timecode;
+		GetDefaultTimecodeInfo(Timecode, FrameRate);
+	}
+
+	return FrameRate;
+}
+
+void UFootageCaptureData::GetDefaultTimecodeInfo(FTimecode& OutTimecode, FFrameRate& OutFrameRate) const
+{
+	for (const FFootageCaptureView& View : Views)
+	{
+		if (View.bImageTimecodePresent)
+		{
+			OutTimecode = View.ImageTimecode;
+			OutFrameRate = View.ImageTimecodeRate;
+			return;
+		}
+	}
+
+	for (const FFootageCaptureView& View : Views)
+	{
+		if (View.bDepthTimecodePresent)
+		{
+			OutTimecode = View.DepthTimecode;
+			OutFrameRate = View.DepthTimecodeRate;
+			return;
+		}
+	}
+
+	if (bAudioTimecodePresent)
+	{
+		OutTimecode = AudioTimecode;
+		OutFrameRate = AudioTimecodeRate;
+		return;
+	}
+
+	OutTimecode = FTimecode(0, 0, 0, 0, false);
+	OutFrameRate = FFrameRate(30, 1);
+}
+
+void UFootageCaptureData::PostLoad()
+{
+	Super::PostLoad();
+
+	if (Metadata.DeviceModel_DEPRECATED != EFootageDeviceClass::Unspecified)
+	{
+		Metadata.DeviceClass = Metadata.DeviceModel_DEPRECATED;
+		Metadata.DeviceModel_DEPRECATED = EFootageDeviceClass::Unspecified;
+	}
+
+	if (CameraCalibration_DEPRECATED != nullptr)
+	{
+		CameraCalibrations.Add(MoveTemp(CameraCalibration_DEPRECATED));
+		CameraCalibration_DEPRECATED = nullptr;
+	}
+
+	if (Audio_DEPRECATED != nullptr)
+	{
+		Audios.Add(MoveTemp(Audio_DEPRECATED));
+		Audio_DEPRECATED = nullptr;
+	}
+}
+
+void UFootageCaptureData::PopulateCameraNames(UFootageCaptureData* InFootageCaptureData, FString& InOutCamera, TArray<TSharedPtr<FString>>& OutCameraNames)
+{
+	OutCameraNames.Reset();
+
+
+
+	if (InFootageCaptureData && !InFootageCaptureData->CameraCalibrations.IsEmpty())
+	{
+		const UCameraCalibration* CameraCalibration = InFootageCaptureData->CameraCalibrations[0];
+		if (CameraCalibration)
+		{
+			for (const FExtendedLensFile& LensFile : CameraCalibration->CameraCalibrations)
+			{
+				OutCameraNames.Add(MakeShared<FString>(LensFile.Name));
+			}
+		}
+	}
+
+	if (OutCameraNames.IsEmpty())
+	{
+		InOutCamera = TEXT("");
+	}
+	else
+	{
+		bool bNameFound = false;
+
+		for (TSharedPtr<FString> Name : OutCameraNames)
+		{
+			if (*Name == InOutCamera)
+			{
+				bNameFound = true;
+				break;
+			}
+		}
+
+		if (!bNameFound)
+		{
+			InOutCamera = *OutCameraNames[0];
+		}
+	}
+}
\ No newline at end of file
diff --git a/Engine/Plugins/VirtualProduction/CaptureData/Source/CaptureDataCore/Private/CaptureDataCoreModule.cpp b/Engine/Plugins/VirtualProduction/CaptureData/Source/CaptureDataCore/Private/CaptureDataCoreModule.cpp
new file mode 100644
--- /dev/null
+++ b/Engine/Plugins/VirtualProduction/CaptureData/Source/CaptureDataCore/Private/CaptureDataCoreModule.cpp
@@ -0,0 +1,63 @@
+// Copyright Epic Games, Inc. All Rights Reserved.
+
+#include "CaptureDataCoreModule.h"
+
+#include "CaptureData.h"
+#include "CaptureDataLog.h"
+
+#include "AssetRegistry/AssetData.h"
+#include "AssetRegistry/AssetRegistryModule.h"
+#include "Misc/MessageDialog.h"
+#include "Modules/ModuleManager.h"
+
+#define LOCTEXT_NAMESPACE "CaptureDataCoreModule"
+
+void FCaptureDataCoreModule::StartupModule()
+{
+	FAssetRegistryModule& AssetRegistryModule = FModuleManager::LoadModuleChecked<FAssetRegistryModule>(TEXT("AssetRegistry"));
+	AssetRegistryModule.Get().OnFilesLoaded().AddRaw(this, &FCaptureDataCoreModule::CheckAssetMigration);
+}
+
+void FCaptureDataCoreModule::ShutdownModule()
+{
+
+}
+
+/*
+Check for MetaHuman assets which need to be migrated due to the source code moving to the Capture Data plugin.
+*/
+void FCaptureDataCoreModule::CheckAssetMigration()
+{
+	FAssetRegistryModule& AssetRegistryModule = FModuleManager::LoadModuleChecked<FAssetRegistryModule>(TEXT("AssetRegistry"));
+	TArray<FAssetData> Assets;
+	
+	// Some MetaHuman assets have moved to the Capture Data plugin so find assets which are referencing the old class path
+	FARFilter Filter;
+	Filter.ClassPaths.Add(FTopLevelAssetPath(FName("/Script/MetaHumanCaptureData"), FName("FootageCaptureData")));
+	Filter.ClassPaths.Add(FTopLevelAssetPath(FName("/Script/MetaHumanCaptureData"), FName("MeshCaptureData")));
+	Filter.ClassPaths.Add(FTopLevelAssetPath(FName("/Script/MetaHumanCore"), FName("MetaHumanCameraCalibration")));
+	Filter.bIncludeOnlyOnDiskAssets = true;
+
+	AssetRegistryModule.Get().GetAssets(Filter, Assets);
+
+	// Prompt user to reload these assets
+	if (!Assets.IsEmpty())
+	{
+		const FText AssetMigrationDialogMessage = LOCTEXT("CaptureDataAssetMigrationDialog", "Found {0} Capture Data assets which need to be reloaded."
+			" Some MetaHuman Animator features may be unavailable until these assets are reloaded. \n\nReload now ? ");
+		const EAppReturnType::Type Answer = FMessageDialog::Open(EAppMsgType::YesNo, FText::Format(AssetMigrationDialogMessage, Assets.Num()));
+		if (Answer == EAppReturnType::No)
+		{
+			return;
+		}
+	}
+
+	for (FAssetData AssetData : Assets)
+	{
+		AssetData.GetAsset()->ReloadConfig();
+	}
+}
+
+IMPLEMENT_MODULE(FCaptureDataCoreModule, CaptureDataCore)
+
+#undef LOCTEXT_NAMESPACE
\ No newline at end of file
diff --git a/Engine/Plugins/VirtualProduction/CaptureData/Source/CaptureDataCore/Private/CaptureDataCoreModule.h b/Engine/Plugins/VirtualProduction/CaptureData/Source/CaptureDataCore/Private/CaptureDataCoreModule.h
new file mode 100644
--- /dev/null
+++ b/Engine/Plugins/VirtualProduction/CaptureData/Source/CaptureDataCore/Private/CaptureDataCoreModule.h
@@ -0,0 +1,18 @@
+// Copyright Epic Games, Inc. All Rights Reserved.
+
+#pragma once
+
+#include "Modules/ModuleInterface.h"
+
+class FCaptureDataCoreModule : public IModuleInterface
+{
+public:
+
+	//~ Begin IModuleInterface interface
+	virtual void StartupModule() override;
+	virtual void ShutdownModule() override;
+	//~ End IModuleInterface interface
+
+private:
+	void CheckAssetMigration();
+};
diff --git a/Engine/Plugins/VirtualProduction/CaptureData/Source/CaptureDataCore/Private/CaptureDataLog.cpp b/Engine/Plugins/VirtualProduction/CaptureData/Source/CaptureDataCore/Private/CaptureDataLog.cpp
new file mode 100644
--- /dev/null
+++ b/Engine/Plugins/VirtualProduction/CaptureData/Source/CaptureDataCore/Private/CaptureDataLog.cpp
@@ -0,0 +1,5 @@
+// Copyright Epic Games, Inc. All Rights Reserved.
+
+#include "CaptureDataLog.h"
+
+DEFINE_LOG_CATEGORY(LogCaptureDataCore);
\ No newline at end of file
diff --git a/Engine/Plugins/VirtualProduction/CaptureData/Source/CaptureDataCore/Private/FrameRange.cpp b/Engine/Plugins/VirtualProduction/CaptureData/Source/CaptureDataCore/Private/FrameRange.cpp
new file mode 100644
--- /dev/null
+++ b/Engine/Plugins/VirtualProduction/CaptureData/Source/CaptureDataCore/Private/FrameRange.cpp
@@ -0,0 +1,25 @@
+// Copyright Epic Games, Inc. All Rights Reserved.
+
+#include "FrameRange.h"
+
+bool FFrameRange::operator==(const FFrameRange& InOther) const
+{
+	return (Name == InOther.Name && StartFrame == InOther.StartFrame && EndFrame == InOther.EndFrame);
+}
+
+bool FFrameRange::ContainsFrame(int32 InFrame, const TArray<FFrameRange>& InFrameRangeArray)
+{
+	for (int32 Index = 0; Index < InFrameRangeArray.Num(); ++Index)
+	{
+		const FFrameRange& FrameRange = InFrameRangeArray[Index];
+
+		if ((FrameRange.StartFrame >= 0 || FrameRange.EndFrame >= 0) &&
+			(FrameRange.StartFrame < 0 || InFrame >= FrameRange.StartFrame) &&
+			(FrameRange.EndFrame < 0 || InFrame <= FrameRange.EndFrame))
+		{
+			return true;
+		}
+	}
+
+	return false;
+}
diff --git a/Engine/Plugins/VirtualProduction/CaptureData/Source/CaptureDataCore/Private/ImageSequencePathChecker.cpp b/Engine/Plugins/VirtualProduction/CaptureData/Source/CaptureDataCore/Private/ImageSequencePathChecker.cpp
new file mode 100644
--- /dev/null
+++ b/Engine/Plugins/VirtualProduction/CaptureData/Source/CaptureDataCore/Private/ImageSequencePathChecker.cpp
@@ -0,0 +1,79 @@
+// Copyright Epic Games, Inc. All Rights Reserved.
+
+#include "ImageSequencePathChecker.h"
+
+#include "CaptureDataLog.h"
+
+#include "Misc/MessageDialog.h"
+
+#define LOCTEXT_NAMESPACE "ImageSequencePathChecker"
+
+namespace UE::CaptureData
+{
+static const FText DialogTitle = LOCTEXT("MissingDataTitle", "Missing Data");
+static const FTextFormat DialogMessageFormat = LOCTEXT(
+	"ImageSequenceMissingDataMessage",
+
+	"{AssetDisplayName} {NumCaptureDataFootageAssets}|plural(one=asset contains,other=assets contain) "
+	"{NumInvalidImageSequences}|plural(one=an image sequence which is,other=image sequences which have) missing data.\n\n"
+	"See output log for details."
+);
+
+FImageSequencePathChecker::FImageSequencePathChecker(FText InAssetDisplayName) :
+	NumCaptureDataFootageAssets(0),
+	NumInvalidImageSequences(0),
+	AssetDisplayName(MoveTemp(InAssetDisplayName))
+{
+}
+
+void FImageSequencePathChecker::Check(const UFootageCaptureData& InCaptureData)
+{
+	++NumCaptureDataFootageAssets;
+
+	const TArray<UFootageCaptureData::FPathAssociation> InvalidImageSequences = InCaptureData.CheckImageSequencePaths();
+	NumInvalidImageSequences += InvalidImageSequences.Num();
+
+	if (!InvalidImageSequences.IsEmpty())
+	{
+		UE_LOG(
+			LogCaptureDataCore,
+			Warning,
+			TEXT("%s contains image sequence(s) with missing data (see below): %s"),
+			*AssetDisplayName.ToString(),
+			*InCaptureData.GetPathName()
+		);
+
+		for (const UFootageCaptureData::FPathAssociation& InvalidImageSequence : InvalidImageSequences)
+		{
+			UE_LOG(
+				LogCaptureDataCore,
+				Warning,
+				TEXT("Image sequence path does not exist or is not a folder: %s (%s)"),
+				*InvalidImageSequence.PathOnDisk,
+				*InvalidImageSequence.AssetPath
+			);
+		}
+	}
+}
+
+void FImageSequencePathChecker::DisplayDialog() const
+{
+	FFormatNamedArguments Args;
+	Args.Add(TEXT("NumCaptureDataFootageAssets"), NumCaptureDataFootageAssets);
+	Args.Add(TEXT("AssetDisplayName"), AssetDisplayName);
+	Args.Add(TEXT("NumInvalidImageSequences"), NumInvalidImageSequences);
+
+	const FText DialogMessage = FText::Format(DialogMessageFormat, Args);
+
+	FMessageDialog::Open(EAppMsgType::Ok, DialogMessage, DialogTitle);
+}
+
+bool FImageSequencePathChecker::HasError() const
+{
+	return NumInvalidImageSequences > 0;
+}
+
+}
+
+#undef LOCTEXT_NAMESPACE
+
diff --git a/Engine/Plugins/VirtualProduction/CaptureData/Source/CaptureDataCore/Private/ImageSequenceUtils.cpp b/Engine/Plugins/VirtualProduction/CaptureData/Source/CaptureDataCore/Private/ImageSequenceUtils.cpp
new file mode 100644
--- /dev/null
+++ b/Engine/Plugins/VirtualProduction/CaptureData/Source/CaptureDataCore/Private/ImageSequenceUtils.cpp
@@ -0,0 +1,135 @@
+// Copyright Epic Games, Inc. All Rights Reserved.
+
+#include "ImageSequenceUtils.h"
+
+#include "IImageWrapper.h"
+#include "IImageWrapperModule.h"
+#include "ImgMediaSource.h"
+#include "Internationalization/Regex.h"
+#include "Misc/FileHelper.h"
+#include "Misc/Paths.h"
+#include "Modules/ModuleManager.h"
+
+
+bool FImageSequenceUtils::GetImageSequencePathAndFiles(const UImgMediaSource* InImgSequence, FString& OutFullSequencePath, TArray<FString>& OutImageFiles)
+{
+	if (InImgSequence == nullptr)
+	{
+		return false;
+	}
+
+	OutFullSequencePath = InImgSequence->GetFullPath();
+
+	return GetImageSequencePathAndFiles(OutFullSequencePath, OutImageFiles);
+}
+
+bool FImageSequenceUtils::GetImageSequencePathAndFiles(const FString& InFullSequencePath, TArray<FString>& OutImageFiles)
+{
+	IImageWrapperModule& ImageWrapperModule = FModuleManager::LoadModuleChecked<IImageWrapperModule>(TEXT("ImageWrapper"));
+	IFileManager& FileManager = IFileManager::Get();
+
+	bool bIterateResult = FileManager.IterateDirectory(*InFullSequencePath, [&OutImageFiles, &ImageWrapperModule](const TCHAR* InFilenameOrDirectory, bool bInIsDirectory)
+	{
+		if (!bInIsDirectory)
+		{
+			EImageFormat Format = ImageWrapperModule.GetImageFormatFromExtension(*FPaths::GetExtension(InFilenameOrDirectory));
+			if (Format != EImageFormat::Invalid)
+			{
+				OutImageFiles.Add(FPaths::GetCleanFilename(InFilenameOrDirectory));
+			}
+		}
+
+		return true;
+	});
+
+	return !OutImageFiles.IsEmpty() && bIterateResult;
+}
+
+bool FImageSequenceUtils::GetImageSequenceInfo(const class UImgMediaSource* InImgSequence, FIntVector2& OutDimensions, int32& OutNumImages)
+{
+	if (InImgSequence == nullptr)
+	{
+		return false;
+	}
+
+	return GetImageSequenceInfo(InImgSequence->GetFullPath(), OutDimensions, OutNumImages);
+}
+
+bool FImageSequenceUtils::GetImageSequenceInfo(const FString& InFullSequencePath, FIntVector2& OutDimensions, int32& OutNumImages)
+{
+	TArray<FString> ImageFiles;
+	bool bFoundImages = GetImageSequencePathAndFiles(InFullSequencePath, ImageFiles);
+
+	if (!bFoundImages)
+	{
+		return false;
+	}
+
+	OutNumImages = ImageFiles.Num();
+
+	const FString SampleImagePath = InFullSequencePath / ImageFiles[0];
+
+	TArray<uint8> RawFileData;
+	if (FFileHelper::LoadFileToArray(RawFileData, *SampleImagePath))
+	{
+		IImageWrapperModule& ImageWrapperModule = FModuleManager::LoadModuleChecked<IImageWrapperModule>(TEXT("ImageWrapper"));
+		EImageFormat ImageFormat = ImageWrapperModule.DetectImageFormat(RawFileData.GetData(), RawFileData.Num());
+		TSharedPtr<IImageWrapper> ImageWrapper = ImageWrapperModule.CreateImageWrapper(ImageFormat);
+
+		if (ImageWrapper.IsValid() && ImageWrapper->SetCompressed(RawFileData.GetData(), RawFileData.Num()))
+		{
+			OutDimensions.X = ImageWrapper->GetWidth();
+			OutDimensions.Y = ImageWrapper->GetHeight();
+			return true;
+		}
+	}
+
+	return false;
+}
+
+bool FImageSequenceUtils::GetTrackingFilePathAndInfo(const class UImgMediaSource* InImgSequence, FString& OutTrackingFilePath, int32& OutFrameOffset, int32& OutNumFrames)
+{
+	return GetTrackingFilePathAndInfo(InImgSequence->GetFullPath(), OutTrackingFilePath, OutFrameOffset, OutNumFrames);
+}
+
+bool FImageSequenceUtils::GetTrackingFilePathAndInfo(const FString& InFullSequencePath, FString& OutTrackingFilePath, int32& OutFrameOffset, int32& OutNumFrames)
+{
+	TArray<FString> ImageFiles;
+	bool bFoundImages = GetImageSequencePathAndFiles(InFullSequencePath, ImageFiles);
+
+	if (ImageFiles.Num() == 0)
+	{
+		return false;
+	}
+
+	ImageFiles.Sort();
+
+	// find an image filename which can be some optional alphabetic or underscore characters followed by some digits 
+	// followed by some optional alphabetic or underscore characters with any extension
+	const FRegexPattern ImageFilenamePattern(TEXT("^[a-zA-Z_]*([0-9]+)[a-zA-Z_]*\\.[a-zA-Z]+$"));
+	FRegexMatcher ImageFilenameMatcher(ImageFilenamePattern, ImageFiles[0]);
+
+	if (ImageFilenameMatcher.FindNext())
+	{
+		const FString Digits = ImageFilenameMatcher.GetCaptureGroup(1);
+		const int32 DigitsStart = ImageFilenameMatcher.GetCaptureGroupBeginning(1);
+		OutFrameOffset = FCString::Atoi(*Digits);
+		const FString InitialChars = ImageFiles[0].Left(DigitsStart);
+		const FString EndChars = ImageFiles[0].Right(ImageFiles[0].Len() - Digits.Len() - DigitsStart);
+		FString DigitsSpecifier = TEXT("%0") + FString::FromInt(Digits.Len()) + TEXT("d");
+		OutTrackingFilePath = InFullSequencePath / InitialChars + DigitsSpecifier + EndChars;
+		OutNumFrames = ImageFiles.Num();
+		return true;
+	}
+
+	return false;
+}
+
+FString FImageSequenceUtils::ExpandFilePathFormat(const FString& InFilePathFormat, int32 InFrameNumber)
+{
+	static_assert(sizeof(TCHAR) == 2); // Ensure TCHAR is the 16-bit wchar_t type expected by swprintf
+
+	TCHAR Buffer[512];
+	swprintf(Buffer, sizeof(Buffer) / sizeof(TCHAR), *InFilePathFormat, InFrameNumber);
+	return Buffer;
+}
diff --git a/Engine/Plugins/VirtualProduction/CaptureData/Source/CaptureDataCore/Private/OpenCVHelperLocal.cpp b/Engine/Plugins/VirtualProduction/CaptureData/Source/CaptureDataCore/Private/OpenCVHelperLocal.cpp
new file mode 100644
--- /dev/null
+++ b/Engine/Plugins/VirtualProduction/CaptureData/Source/CaptureDataCore/Private/OpenCVHelperLocal.cpp
@@ -0,0 +1,60 @@
+// Copyright Epic Games, Inc. All Rights Reserved.
+
+#include "OpenCVHelperLocal.h"
+#include "Math/TransformCalculus3D.h"
+
+void FOpenCVHelperLocal::ConvertCoordinateSystem(FTransform& Transform, const EAxis SrcXInDstAxis, const EAxis SrcYInDstAxis, const EAxis SrcZInDstAxis)
+{
+	// Unreal Engine:
+	//   Front : X
+	//   Right : Y
+	//   Up    : Z
+	//
+	// OpenCV:
+	//   Front : Z
+	//   Right : X
+	//   Up    : Yn
+
+	FMatrix M12 = FMatrix::Identity;
+
+	M12.SetColumn(0, UnitVectorFromAxisEnum(SrcXInDstAxis));
+	M12.SetColumn(1, UnitVectorFromAxisEnum(SrcYInDstAxis));
+	M12.SetColumn(2, UnitVectorFromAxisEnum(SrcZInDstAxis));
+
+	Transform.SetFromMatrix(M12.GetTransposed() * Transform.ToMatrixWithScale() * M12);
+}
+
+void FOpenCVHelperLocal::ConvertUnrealToOpenCV(FTransform& Transform)
+{
+	ConvertCoordinateSystem(Transform, EAxis::Y, EAxis::Zn, EAxis::X);
+}
+
+void FOpenCVHelperLocal::ConvertOpenCVToUnreal(FTransform& Transform)
+{
+	ConvertCoordinateSystem(Transform, EAxis::Z, EAxis::X, EAxis::Yn);
+}
+
+FVector FOpenCVHelperLocal::ConvertUnrealToOpenCV(const FVector& Vector)
+{
+	return FVector(Vector.Y, -Vector.Z, Vector.X);
+}
+
+FVector FOpenCVHelperLocal::ConvertOpenCVToUnreal(const FVector& Vector)
+{
+	return FVector(Vector.Z, Vector.X, -Vector.Y);
+}
+
+void FOpenCVHelperLocal::ConvertOpenCVToUnreal(const FMatrix& InRotationOpenCV, const FVector& InTranslationOpenCV, FRotator& OutRotatorUE, FVector& OutTranslationUE)
+{
+	// this function converts a rotation and translation from OpenCV coordinate frame, where we have
+	// right-handed coordinate system, x right, y down to
+	// UE coordinate system, where we have left-handed coordinate system, y right, z up
+
+	FVector XAxisUE = FVector(InRotationOpenCV.M[2][2], InRotationOpenCV.M[2][0], -InRotationOpenCV.M[2][1]); // comes from Z axis in rotation matrix 
+	FVector YAxisUE = FVector(InRotationOpenCV.M[0][2], InRotationOpenCV.M[0][0], -InRotationOpenCV.M[0][1]); // comes from X axis in rotation matrix
+	FMatrix RotationUE = FRotationMatrix::MakeFromXY(XAxisUE, YAxisUE);
+
+	// again swap the elements in the position to get the result in UE coordinate space
+	OutTranslationUE = FVector(InTranslationOpenCV.Z, InTranslationOpenCV.X, -InTranslationOpenCV.Y);
+	OutRotatorUE = TransformConverter<FRotator>::Convert(RotationUE);
+}
diff --git a/Engine/Plugins/VirtualProduction/CaptureData/Source/CaptureDataCore/Private/Tests/FootageCaptureMetadata.spec.cpp b/Engine/Plugins/VirtualProduction/CaptureData/Source/CaptureDataCore/Private/Tests/FootageCaptureMetadata.spec.cpp
new file mode 100644
--- /dev/null
+++ b/Engine/Plugins/VirtualProduction/CaptureData/Source/CaptureDataCore/Private/Tests/FootageCaptureMetadata.spec.cpp
@@ -0,0 +1,82 @@
+// Copyright Epic Games, Inc. All Rights Reserved.
+
+#include "CaptureData.h"
+
+#include "CoreTypes.h"
+#include "Misc/AutomationTest.h"
+#include "UObject/Class.h"
+
+DEFINE_LOG_CATEGORY_STATIC(LogFootageCaptureMetadataSpec, Log, All);
+
+BEGIN_DEFINE_SPEC(FootageCaptureMetadataSpec, "MetaHuman.Capture.FootageCaptureMetadata"
+	, EAutomationTestFlags_ApplicationContextMask
+	| EAutomationTestFlags::ProductFilter)
+	FFootageCaptureMetadata Metadata;
+	TMap<FString, EFootageDeviceClass> ModelNameToDeviceClass;
+END_DEFINE_SPEC(FootageCaptureMetadataSpec)
+void FootageCaptureMetadataSpec::Define()
+{	
+	Describe("SetDeviceClass", [this]()
+	{
+		// https://theapplewiki.com/wiki/Models
+		ModelNameToDeviceClass.Reset();
+		ModelNameToDeviceClass.Add("iPhone10,1", EFootageDeviceClass::iPhone11OrEarlier); // iPhone 8 (1)
+		ModelNameToDeviceClass.Add("iPhone10,2", EFootageDeviceClass::iPhone11OrEarlier); // iPhone 8 Plus (1)
+		ModelNameToDeviceClass.Add("iPhone10,3", EFootageDeviceClass::iPhone11OrEarlier); // iPhone X (1)
+		ModelNameToDeviceClass.Add("iPhone10,4", EFootageDeviceClass::iPhone11OrEarlier); // iPhone 8 (2)
+		ModelNameToDeviceClass.Add("iPhone10,5", EFootageDeviceClass::iPhone11OrEarlier); // iPhone Plus (2)
+		ModelNameToDeviceClass.Add("iPhone10,6", EFootageDeviceClass::iPhone11OrEarlier); // iPhone X (2)
+		ModelNameToDeviceClass.Add("iPhone11,2", EFootageDeviceClass::iPhone11OrEarlier); // iPhone XS
+		ModelNameToDeviceClass.Add("iPhone11,4", EFootageDeviceClass::iPhone11OrEarlier); // iPhone XS Max (1)
+		ModelNameToDeviceClass.Add("iPhone11,6", EFootageDeviceClass::iPhone11OrEarlier); // iPhone XS Max (2)
+		ModelNameToDeviceClass.Add("iPhone11,8", EFootageDeviceClass::iPhone11OrEarlier); // iPhone XR
+		ModelNameToDeviceClass.Add("iPhone12,1", EFootageDeviceClass::iPhone11OrEarlier); // iPhone 11
+		ModelNameToDeviceClass.Add("iPhone12,3", EFootageDeviceClass::iPhone11OrEarlier); // iPhone 11 Pro
+		ModelNameToDeviceClass.Add("iPhone12,5", EFootageDeviceClass::iPhone11OrEarlier); // iPhone 11 Pro Max
+		ModelNameToDeviceClass.Add("iPhone12,8", EFootageDeviceClass::OtheriOSDevice);    // iPhone SE 2
+		ModelNameToDeviceClass.Add("iPhone13,1", EFootageDeviceClass::iPhone12);          // iPhone 12 Mini
+		ModelNameToDeviceClass.Add("iPhone13,2", EFootageDeviceClass::iPhone12);          // iPhone 12
+		ModelNameToDeviceClass.Add("iPhone13,3", EFootageDeviceClass::iPhone12);          // iPhone 12 Pro
+		ModelNameToDeviceClass.Add("iPhone13,4", EFootageDeviceClass::iPhone12);          // iPhone 12 Pro Max
+		ModelNameToDeviceClass.Add("iPhone14,2", EFootageDeviceClass::iPhone13);          // iPhone 13 Pro
+		ModelNameToDeviceClass.Add("iPhone14,3", EFootageDeviceClass::iPhone13);          // iPhone 13 Pro Max
+		ModelNameToDeviceClass.Add("iPhone14,4", EFootageDeviceClass::iPhone13);          // iPhone 13 Mini
+		ModelNameToDeviceClass.Add("iPhone14,5", EFootageDeviceClass::iPhone13);          // iPhone 13
+		ModelNameToDeviceClass.Add("iPhone14,6", EFootageDeviceClass::OtheriOSDevice);    // iPhone SE 3
+		ModelNameToDeviceClass.Add("iPhone14,7", EFootageDeviceClass::iPhone14OrLater);   // iPhone 14
+		ModelNameToDeviceClass.Add("iPhone14,8", EFootageDeviceClass::iPhone14OrLater);   // iPhone 14 Plus
+		ModelNameToDeviceClass.Add("iPhone15,2", EFootageDeviceClass::iPhone14OrLater);   // iPhone 14 Pro
+		ModelNameToDeviceClass.Add("iPhone15,3", EFootageDeviceClass::iPhone14OrLater);   // iPhone 14 Pro Max
+		ModelNameToDeviceClass.Add("iPhone15,4", EFootageDeviceClass::iPhone14OrLater);   // iPhone 15
+		ModelNameToDeviceClass.Add("iPhone15,5", EFootageDeviceClass::iPhone14OrLater);   // iPhone 15 Plus
+		ModelNameToDeviceClass.Add("iPhone16,1", EFootageDeviceClass::iPhone14OrLater);   // iPhone 15 Pro
+		ModelNameToDeviceClass.Add("iPhone16,2", EFootageDeviceClass::iPhone14OrLater);   // iPhone 15 Pro Max
+		ModelNameToDeviceClass.Add("iPhone99,1", EFootageDeviceClass::iPhone14OrLater);   // Future iPhone
+		ModelNameToDeviceClass.Add("iPhone999,9", EFootageDeviceClass::iPhone14OrLater);  // Distant Future iPhone
+		ModelNameToDeviceClass.Add("iPhone", EFootageDeviceClass::OtheriOSDevice);        // Invalid iPhone Model Number
+		ModelNameToDeviceClass.Add("iPhone1,2,3", EFootageDeviceClass::OtheriOSDevice);   // Invalid iPhone Model Number
+		
+		ModelNameToDeviceClass.Add("iPad8,1", EFootageDeviceClass::OtheriOSDevice);       // iPad Pro 11-inch 1
+		ModelNameToDeviceClass.Add("iPad11,3", EFootageDeviceClass::OtheriOSDevice);      // iPad Air 3
+		ModelNameToDeviceClass.Add("iPad13,11", EFootageDeviceClass::OtheriOSDevice);     // iPad Pro 12.9-inch 5
+		ModelNameToDeviceClass.Add("iPad", EFootageDeviceClass::OtheriOSDevice);          // Invalid iPad Model Number
+
+		ModelNameToDeviceClass.Add("StereoHMC", EFootageDeviceClass::StereoHMC);          // Stereo HMC
+
+		ModelNameToDeviceClass.Add("iTablet", EFootageDeviceClass::Unspecified);          // Invalid Model Number
+		ModelNameToDeviceClass.Add("NotAnIPhone", EFootageDeviceClass::Unspecified);      // Invalid Model Number
+
+		for (const TPair<FString, EFootageDeviceClass>& Mapping : ModelNameToDeviceClass)
+		{
+			const FString ModelName = Mapping.Key;
+			const EFootageDeviceClass Expected = Mapping.Value;
+			const FString ExpectedName = UEnum::GetDisplayValueAsText(Expected).ToString();
+			It(FString::Printf(TEXT("should set DeviceClass to '%s' when DeviceModel is '%s'"), *ExpectedName, *ModelName), [this, ModelName, Expected]
+			{
+				Metadata.SetDeviceClass(ModelName);
+				const EFootageDeviceClass Actual = Metadata.DeviceClass;
+				TestEqual("DeviceClass", Actual, Expected);
+			});
+		}
+	});
+}
\ No newline at end of file
diff --git a/Engine/Plugins/VirtualProduction/CaptureData/Source/CaptureDataCore/Public/CameraCalibration.h b/Engine/Plugins/VirtualProduction/CaptureData/Source/CaptureDataCore/Public/CameraCalibration.h
new file mode 100644
--- /dev/null
+++ b/Engine/Plugins/VirtualProduction/CaptureData/Source/CaptureDataCore/Public/CameraCalibration.h
@@ -0,0 +1,100 @@
+// Copyright Epic Games, Inc. All Rights Reserved.
+
+#pragma once
+
+#include "UObject/ObjectMacros.h"
+#include "EditorFramework/AssetImportData.h"
+
+#include "CameraCalibration.generated.h"
+
+
+USTRUCT(BlueprintType)
+struct FExtendedLensFile
+{
+	GENERATED_BODY()
+
+
+public:
+	UPROPERTY(EditAnywhere, Category = Calibration)
+	FString Name;
+
+	UPROPERTY(EditAnywhere, Category = Calibration)
+	bool IsDepthCamera{};
+
+	UPROPERTY(EditAnywhere, Category = Calibration)
+	TObjectPtr<class ULensFile> LensFile{};
+};
+
+USTRUCT(BlueprintType)
+struct FStereoPair
+{
+	GENERATED_BODY()
+
+
+public:
+	UPROPERTY(EditAnywhere, Category = Calibration)
+	uint32 CameraIndex1 {};
+
+	UPROPERTY(EditAnywhere, Category = Calibration)
+	uint32 CameraIndex2 {};
+};
+
+struct CAPTUREDATACORE_API FCameraCalibration
+{
+public:
+	FString Name;
+	FString Camera;
+	int32 ImageSizeX = 0;
+	int32 ImageSizeY = 0;
+	double FX = 0;
+	double FY = 0;
+	double CX = 0;
+	double CY = 0;
+	double K1 = 0;
+	double K2 = 0;
+	double P1 = 0;
+	double P2 = 0;
+	double K3 = 0;
+	double K4 = 0;
+	double K5 = 0;
+	double K6 = 0;
+	FMatrix Transform;
+};
+
+/** Camera Calibration Asset
+*
+*   Contains the parameters for calibrating the camera
+*   used in footage for MetaHuman Identity and Performance assets.
+* 
+**/
+UCLASS(BlueprintType)
+class CAPTUREDATACORE_API UCameraCalibration : public UObject
+{
+	GENERATED_BODY()
+
+public:
+	//~Begin UObject interface
+	virtual void PostInitProperties() override;
+	virtual void PostLoad() override;
+	virtual void GetAssetRegistryTags(TArray<FAssetRegistryTag>& OutTags) const override;
+	//~End UObject interface
+
+	UPROPERTY(EditAnywhere, Category = Calibration)
+	TArray<FExtendedLensFile> CameraCalibrations;
+
+	UPROPERTY(EditAnywhere, Category = Calibration)
+	TArray<FStereoPair> StereoPairs;
+
+#if WITH_EDITORONLY_DATA
+	/** Importing data and options used for importing mhaical files */
+	UPROPERTY(VisibleAnywhere, Instanced, Category = "Import Settings")
+	TObjectPtr<UAssetImportData> AssetImportData;
+#endif // WITH_EDITORONLY_DATA
+
+	bool ConvertToTrackerNodeCameraModels(TArray<FCameraCalibration>& OutCalibrations,
+										  TArray<TPair<FString, FString>>& OutStereoReconstructionPairs) const;
+
+	bool ConvertFromTrackerNodeCameraModels(TArray<FCameraCalibration>& InCalibrations);
+
+	int32 GetIndexByCameraName(const FString& InName) const;
+};
\ No newline at end of file
diff --git a/Engine/Plugins/VirtualProduction/CaptureData/Source/CaptureDataCore/Public/CaptureData.h b/Engine/Plugins/VirtualProduction/CaptureData/Source/CaptureDataCore/Public/CaptureData.h
new file mode 100644
--- /dev/null
+++ b/Engine/Plugins/VirtualProduction/CaptureData/Source/CaptureDataCore/Public/CaptureData.h
@@ -0,0 +1,284 @@
+// Copyright Epic Games, Inc. All Rights Reserved.
+
+#pragma once
+
+#include "UObject/Object.h"
+#include "UObject/ObjectMacros.h"
+#include "IntVectorTypes.h"
+#include "Misc/Timecode.h"
+#include "Misc/Paths.h"
+#include "Engine/StaticMesh.h"
+#include "FrameRange.h"
+
+#include "Templates/ValueOrError.h"
+
+#include "CaptureData.generated.h"
+
+/////////////////////////////////////////////////////
+// UCaptureData
+
+// Delegate called when something changes in the capture data that others should know about
+DECLARE_MULTICAST_DELEGATE(FOnCaptureDataInternalsChanged)
+
+UCLASS(Abstract, MinimalAPI, BlueprintType)
+class UCaptureData : public UObject
+{
+	GENERATED_BODY()
+
+public:
+	/** Returns true is the capture data is fully initialized with all required information present */
+	virtual bool IsInitialized() const PURE_VIRTUAL(UCaptureData::IsInitialized, return false;);
+
+	FOnCaptureDataInternalsChanged& OnCaptureDataInternalsChanged()
+	{
+		return OnCaptureDataInternalsChangedDelegate;
+	}
+
+#if WITH_EDITOR
+
+	//~Begin UObject interface
+	virtual void PostEditChangeProperty(struct FPropertyChangedEvent& InPropertyChangedEvent) override;
+	virtual void PostTransacted(const FTransactionObjectEvent& InTransactionEvent) override;
+	//~End UObject interface
+
+#endif
+
+protected:
+	/** Notify that something internal to the capture data changed */
+	void NotifyInternalsChanged();
+
+private:
+	FOnCaptureDataInternalsChanged OnCaptureDataInternalsChangedDelegate;
+};
+
+/** Capture Data (Mesh) Asset
+*
+*   An asset that contains the Mesh data representing a facial
+*   expression (Pose), that can be used in MetaHuman Identity
+*   to generate a Skeletal Mesh or a full MetaHuman resembling
+*   a real person or a sculpted character.
+*
+*   If pose data represents a real person, the resulting
+*   Skeletal Mesh can be used to generate animation from footage
+*   of that person in MetaHuman Performance asset
+*
+**/
+UCLASS()
+class CAPTUREDATACORE_API UMeshCaptureData : public UCaptureData
+{
+	GENERATED_BODY()
+
+public:
+	// UCaptureData interface
+	virtual bool IsInitialized() const override;
+
+public:
+	/** Gets the data in the format expected by the face fitting API */
+	void GetDataForConforming(const FTransform& InTransform, TArray<float>& OutVertices, TArray<int32>& OutTriangles) const;
+
+	// The target mesh for conforming. This can be either a Static or SkeletalMesh
+	UPROPERTY(EditAnywhere, Category = "Capture", Meta = (AllowedClasses = "/Script/Engine.StaticMesh, /Script/Engine.SkeletalMesh"))
+	TObjectPtr<class UObject> TargetMesh;
+};
+
+/////////////////////////////////////////////////////
+// UFootageCaptureData
+
+UENUM()
+enum class EFootageDeviceClass : uint8
+{
+	Unspecified,
+	iPhone11OrEarlier	UMETA(DisplayName = "iPhone 11 or earlier"),
+	iPhone12			UMETA(DisplayName = "iPhone 12"),
+	iPhone13			UMETA(DisplayName = "iPhone 13"),
+	iPhone14OrLater		UMETA(DisplayName = "iPhone 14 or later"),
+	OtheriOSDevice		UMETA(DisplayName = "Other iOS device"),
+	StereoHMC			UMETA(DisplayName = "Stereo HMC"),
+};
+
+USTRUCT(BlueprintType)
+struct CAPTUREDATACORE_API FFootageCaptureMetadata
+{
+	GENERATED_BODY()
+
+	UPROPERTY(EditAnywhere, BlueprintReadWrite, Category = "Metadata")
+	int32 Width = 0;
+
+	UPROPERTY(EditAnywhere, BlueprintReadWrite, Category = "Metadata")
+	int32 Height = 0;
+
+	UPROPERTY(EditAnywhere, BlueprintReadWrite, Category = "Metadata")
+	double FrameRate = 0.0;
+
+	UPROPERTY(EditAnywhere, BlueprintReadWrite, Category = "Metadata")
+	EFootageDeviceClass DeviceClass = EFootageDeviceClass::Unspecified;
+
+	UPROPERTY(meta = (DeprecatedProperty, DeprecationMessage = "Property has been renamed to Device Class"))
+	EFootageDeviceClass DeviceModel_DEPRECATED = EFootageDeviceClass::Unspecified;
+
+	UPROPERTY(EditAnywhere, BlueprintReadWrite, DisplayName = "Device Model", Category = "Metadata")
+	FString DeviceModelName;
+
+	/** Sets the DeviceClass property parsing the model string accordingly */
+	void SetDeviceClass(const FString& InDeviceModel);
+
+private:
+	struct FIosDeviceVersion
+	{
+		uint16 Major;
+		uint16 Minor;
+
+		FIosDeviceVersion(uint16 Major, uint16 Minor)
+			: Major(Major), Minor(Minor)
+		{}
+	};
+	TOptional<FIosDeviceVersion> ParseIosDeviceVersion(const FString& Prefix, const FString& ModelName);
+	EFootageDeviceClass IPhoneDeviceClass(TOptional<FIosDeviceVersion> IosDeviceVersion);
+};
+
+USTRUCT(BlueprintType)
+struct FFootageCaptureView
+{
+	GENERATED_BODY()
+
+	UPROPERTY(EditAnywhere, BlueprintReadWrite, Category = "View")
+	TObjectPtr<class UImgMediaSource> ImageSequence;
+
+	UPROPERTY(EditAnywhere, BlueprintReadWrite, Category = "View")
+	bool bImageTimecodePresent = false;
+
+	UPROPERTY(EditAnywhere, BlueprintReadWrite, Category = "View")
+	FTimecode ImageTimecode = FTimecode(0, 0, 0, 0, false);
+
+	UPROPERTY(EditAnywhere, BlueprintReadWrite, Category = "View")
+	FFrameRate ImageTimecodeRate = FFrameRate(30, 1);
+
+	UPROPERTY(EditAnywhere, BlueprintReadWrite, Category = "View")
+	TObjectPtr<class UImgMediaSource> DepthSequence;
+
+	UPROPERTY(EditAnywhere, BlueprintReadWrite, Category = "View")
+	bool bDepthTimecodePresent = false;
+
+	UPROPERTY(EditAnywhere, BlueprintReadWrite, Category = "View")
+	FTimecode DepthTimecode = FTimecode(0, 0, 0, 0, false);
+
+	UPROPERTY(EditAnywhere, BlueprintReadWrite, Category = "View")
+	FFrameRate DepthTimecodeRate = FFrameRate(30, 1);
+};
+
+UENUM()
+enum class ETimecodeAlignment : int32
+{
+	None            UMETA(DisplayName = "None"),
+	Absolute        UMETA(DisplayName = "Absolute"),
+	Relative        UMETA(DisplayName = "Relative"),
+};
+
+/**
+ * Capture Data (Footage) Asset
+ *
+ *  An asset that contains the footage data showing facial
+ *  expressions (Poses) that can be used by MetaHuman Identity
+ *  Asset toolkit to generate a Skeletal Mesh or a MetaHuman
+ *  resembling a real person or a sculpted character.
+ *
+ *  The resulting Skeletal Mesh resembling the actor from the
+ *  footage can be used for generating animation from footage
+ *  of that person in MetaHuman Performance asset
+ *
+**/
+UCLASS()
+class CAPTUREDATACORE_API UFootageCaptureData : public UCaptureData
+{
+	GENERATED_BODY()
+
+public:
+
+	using FVerifyResult = TValueOrError<void, FString>;
+
+	//~Begin UCaptureData interface
+	virtual bool IsInitialized() const override;
+	//~End UCaptureData interface
+
+	//~Begin UObject interface
+	virtual void PostLoad() override;
+	//~End UObject interface
+
+public:
+
+	UPROPERTY(EditAnywhere, BlueprintReadWrite, Category = "Capture")
+	TArray<FFootageCaptureView> Views;
+
+	UPROPERTY(EditAnywhere, BlueprintReadWrite, Category = "Capture")
+	TArray<TObjectPtr<class USoundWave>> Audios;
+
+	UPROPERTY(meta = (DeprecatedProperty, DeprecationMessage = "Audio is deprecated. Please use Audios instead."))
+	TObjectPtr<class USoundWave> Audio_DEPRECATED;
+
+	UPROPERTY(EditAnywhere, BlueprintReadWrite, Category = "Capture")
+	bool bAudioTimecodePresent = false;
+
+	UPROPERTY(EditAnywhere, BlueprintReadWrite, Category = "Capture")
+	FTimecode AudioTimecode = FTimecode(0, 0, 0, 0, false);
+
+	UPROPERTY(EditAnywhere, BlueprintReadWrite, Category = "Capture")
+	FFrameRate AudioTimecodeRate = FFrameRate(30, 1);
+
+	UPROPERTY(EditAnywhere, BlueprintReadWrite, Category = "Capture")
+	TArray<TObjectPtr<class UCameraCalibration>> CameraCalibrations;
+
+	UPROPERTY(EditAnywhere, BlueprintReadWrite, Category = "Capture")
+	FFootageCaptureMetadata Metadata;
+
+	UPROPERTY(EditAnywhere, BlueprintReadWrite, Category = "Capture", DisplayName = "Excluded Frames")
+	TArray<FFrameRange> CaptureExcludedFrames;
+
+	UPROPERTY(meta = (DeprecatedProperty, DeprecationMessage = "CameraCalibration is deprecated. Please use CameraCalibrations instead."))
+	TObjectPtr<class UCameraCalibration> CameraCalibration_DEPRECATED;
+
+public:
+
+	/** Gets the resolution of the color channel */
+	FIntPoint GetFootageColorResolution() const;
+
+	void GetFrameRanges(const FFrameRate& InTargetRate, ETimecodeAlignment InTimecodeAlignment, bool bInIncludeAudio, TMap<TWeakObjectPtr<UObject>, TRange<FFrameNumber>>& OutMediaFrameRanges, TRange<FFrameNumber>& OutProcessingFrameRange, TRange<FFrameNumber>& OutMaximumFrameRange) const;
+
+	static TRange<FFrameNumber> GetAudioFrameRange(const FFrameRate& InTargetRate, ETimecodeAlignment InTimecodeAlignment, USoundWave* InMedia, const FTimecode& InMediaTimecode, const FFrameRate& InMediaTimecodeRate);
+
+	/** Functions for getting the effective timecode info on a track - note some tracks may not have a timecode set
+	*   in which case their effective timecode is defined by others tracks that do have a timecode set.
+	*/
+	FTimecode GetEffectiveImageTimecode(int32 InView) const;
+	FFrameRate GetEffectiveImageTimecodeRate(int32 InView) const;
+
+	FTimecode GetEffectiveDepthTimecode(int32 InView) const;
+	FFrameRate GetEffectiveDepthTimecodeRate(int32 InView) const;
+
+	FTimecode GetEffectiveAudioTimecode() const;
+	FFrameRate GetEffectiveAudioTimecodeRate() const;
+
+	FVerifyResult VerifyData() const;
+
+	struct FPathAssociation
+	{
+		FPathAssociation(const FString& InPathOnDisk, const FString& InAssetPath);
+
+		FString PathOnDisk;
+		FString AssetPath;
+	};
+
+	/** Check the existence of the parent directory for all image sequences and return a list of those which fail the test. */
+	[[nodiscard]] TArray<FPathAssociation> CheckImageSequencePaths() const;
+
+	/** List of all RGB cameras (views) in the footage capture data. Ensures InOutCamera is a valid camera name. */
+	static void PopulateCameraNames(UFootageCaptureData* InFootageCaptureData, FString& InOutCamera, TArray<TSharedPtr<FString>>& OutCameraNames);
+
+private:
+
+	/** Checks if camera calibration and views are valid. */
+	FVerifyResult ViewsContainsValidData() const;
+	FVerifyResult MetadataContainsValidData() const;
+	FVerifyResult CalibrationContainsValidData() const;
+
+	void GetDefaultTimecodeInfo(FTimecode& OutTimecode, FFrameRate& OutFrameRate) const;
+};
\ No newline at end of file
diff --git a/Engine/Plugins/VirtualProduction/CaptureData/Source/CaptureDataCore/Public/CaptureDataLog.h b/Engine/Plugins/VirtualProduction/CaptureData/Source/CaptureDataCore/Public/CaptureDataLog.h
new file mode 100644
--- /dev/null
+++ b/Engine/Plugins/VirtualProduction/CaptureData/Source/CaptureDataCore/Public/CaptureDataLog.h
@@ -0,0 +1,7 @@
+// Copyright Epic Games, Inc. All Rights Reserved.
+
+#pragma once
+
+#include "Logging/LogMacros.h"
+
+DECLARE_LOG_CATEGORY_EXTERN(LogCaptureDataCore, Log, All);
diff --git a/Engine/Plugins/VirtualProduction/CaptureData/Source/CaptureDataCore/Public/FrameRange.h b/Engine/Plugins/VirtualProduction/CaptureData/Source/CaptureDataCore/Public/FrameRange.h
new file mode 100644
--- /dev/null
+++ b/Engine/Plugins/VirtualProduction/CaptureData/Source/CaptureDataCore/Public/FrameRange.h
@@ -0,0 +1,40 @@
+// Copyright Epic Games, Inc. All Rights Reserved.
+
+#pragma once
+
+#include "CoreMinimal.h"
+
+#include "FrameRange.generated.h"
+
+USTRUCT(BlueprintType)
+struct CAPTUREDATACORE_API FFrameRange
+{
+
+public:
+
+	GENERATED_BODY()
+
+	bool operator==(const FFrameRange& InOther) const;
+
+	UPROPERTY(EditAnywhere, Category = "Frame Range")
+	FString Name = TEXT("Unnamed");
+
+	UPROPERTY(EditAnywhere, Category = "Frame Range")
+	int32 StartFrame = -1;
+
+	UPROPERTY(EditAnywhere, Category = "Frame Range")
+	int32 EndFrame = -1;
+
+	static bool ContainsFrame(int32 InFrame, const TArray<FFrameRange>& InFrameRangeArray);
+};
+
+UENUM()
+enum class EFrameRangeType : uint8
+{
+	UserExcluded		UMETA(DisplayName = "User Excluded Frame"),
+	ProcessingExcluded	UMETA(DisplayName = "Processing Excluded Frame"),
+	CaptureExcluded		UMETA(DisplayName = "Capture Excluded Frame"),
+	None
+};
+
+typedef TMap<EFrameRangeType, TArray<FFrameRange>> FFrameRangeMap;
diff --git a/Engine/Plugins/VirtualProduction/CaptureData/Source/CaptureDataCore/Public/ImageSequencePathChecker.h b/Engine/Plugins/VirtualProduction/CaptureData/Source/CaptureDataCore/Public/ImageSequencePathChecker.h
new file mode 100644
--- /dev/null
+++ b/Engine/Plugins/VirtualProduction/CaptureData/Source/CaptureDataCore/Public/ImageSequencePathChecker.h
@@ -0,0 +1,27 @@
+// Copyright Epic Games, Inc. All Rights Reserved.
+
+#pragma once
+
+#include "CaptureData.h"
+
+#include "Internationalization/Internationalization.h"
+
+namespace UE::CaptureData
+{
+
+class CAPTUREDATACORE_API FImageSequencePathChecker
+{
+public:
+	explicit FImageSequencePathChecker(FText InAssetDisplayName);
+
+	void Check(const UFootageCaptureData& InCaptureData);
+	void DisplayDialog() const;
+	bool HasError() const;
+
+private:
+	int32 NumCaptureDataFootageAssets;
+	int32 NumInvalidImageSequences;
+	FText AssetDisplayName;
+};
+
+}
diff --git a/Engine/Plugins/VirtualProduction/CaptureData/Source/CaptureDataCore/Public/ImageSequenceUtils.h b/Engine/Plugins/VirtualProduction/CaptureData/Source/CaptureDataCore/Public/ImageSequenceUtils.h
new file mode 100644
--- /dev/null
+++ b/Engine/Plugins/VirtualProduction/CaptureData/Source/CaptureDataCore/Public/ImageSequenceUtils.h
@@ -0,0 +1,20 @@
+// Copyright Epic Games, Inc. All Rights Reserved.
+
+#pragma once
+
+#include "CoreMinimal.h"
+
+class CAPTUREDATACORE_API FImageSequenceUtils
+{
+public:
+	static bool GetImageSequencePathAndFiles(const class UImgMediaSource* InImgSequence, FString& OutFullSequencePath, TArray<FString>& OutImageFiles);
+	static bool GetImageSequencePathAndFiles(const FString& InFullSequencePath, TArray<FString>& OutImageFiles);
+
+	static bool GetImageSequenceInfo(const class UImgMediaSource* InImgSequence, FIntVector2& OutDimensions, int32& OutNumImages);
+	static bool GetImageSequenceInfo(const FString& InFullSequencePath, FIntVector2& OutDimensions, int32& OutNumImages);
+
+	static bool GetTrackingFilePathAndInfo(const class UImgMediaSource* InImgSequence, FString& OutTrackingFilePath, int32& OutFrameOffset, int32& OutNumFrames);
+	static bool GetTrackingFilePathAndInfo(const FString& InFullSequencePath, FString& OutTrackingFilePath, int32& OutFrameOffset, int32& OutNumFrames);
+
+	static FString ExpandFilePathFormat(const FString& InFilePathFormat, int32 InFrameNumber);
+};
diff --git a/Engine/Plugins/VirtualProduction/CaptureData/Source/CaptureDataCore/Public/OpenCVHelperLocal.h b/Engine/Plugins/VirtualProduction/CaptureData/Source/CaptureDataCore/Public/OpenCVHelperLocal.h
new file mode 100644
--- /dev/null
+++ b/Engine/Plugins/VirtualProduction/CaptureData/Source/CaptureDataCore/Public/OpenCVHelperLocal.h
@@ -0,0 +1,52 @@
+// Copyright Epic Games, Inc. All Rights Reserved.
+
+#pragma once
+#include "CoreMinimal.h"
+
+// NB this is a cut-down version of the FOpenCVHelper class from OpenCV runtime plugin
+class CAPTUREDATACORE_API FOpenCVHelperLocal
+{
+public:
+	/** Enumeration to specify any cartesian axis in positive or negative directions */
+	enum class EAxis
+	{
+		X, Y, Z,
+		Xn, Yn, Zn,
+	};
+
+	// These axes must match the order in which they are declared in EAxis
+	inline static const TArray<FVector> UnitVectors =
+	{
+		{  1,  0,  0 }, //  X
+		{  0,  1,  0 }, //  Y
+		{  0,  0,  1 }, //  Z
+		{ -1,  0,  0 }, // -X
+		{  0, -1,  0 }, // -Y
+		{  0,  0, -1 }, // -Z
+	};
+
+	static const FVector& UnitVectorFromAxisEnum(const EAxis Axis)
+	{
+		return UnitVectors[std::underlying_type_t<EAxis>(Axis)];
+	};
+
+	/** Converts in-place the coordinate system of the given FTransform by specifying the source axes in terms of the destination axes */
+	static void ConvertCoordinateSystem(FTransform& Transform, const EAxis DstXInSrcAxis, const EAxis DstYInSrcAxis, const EAxis DstZInSrcAxis);
+
+	/** Converts in-place an FTransform in Unreal coordinates to OpenCV coordinates */
+	static void ConvertUnrealToOpenCV(FTransform& Transform);
+
+	/** Converts in-place an FTransform in OpenCV coordinates to Unreal coordinates */
+	static void ConvertOpenCVToUnreal(FTransform& Transform);
+
+	/** Converts an FVector in Unreal coordinates to OpenCV coordinates */
+	static FVector ConvertUnrealToOpenCV(const FVector& Transform);
+
+	/** Converts an FTransform in OpenCV coordinates to Unreal coordinates */
+	static FVector ConvertOpenCVToUnreal(const FVector& Transform);
+
+	/** Converts a FMatrix and FVector in OpenCV coordinates to FRotator and FVector in Unreal coordinates */
+	static void ConvertOpenCVToUnreal(const FMatrix& InRotationOpenCV, const FVector& InTranslationOpenCV, FRotator& OutRotatorUE, FVector& OutTranslationUE);
+
+};
+
diff --git a/Engine/Plugins/VirtualProduction/CaptureData/Source/CaptureDataEditor/CaptureDataEditor.Build.cs b/Engine/Plugins/VirtualProduction/CaptureData/Source/CaptureDataEditor/CaptureDataEditor.Build.cs
new file mode 100644
--- /dev/null
+++ b/Engine/Plugins/VirtualProduction/CaptureData/Source/CaptureDataEditor/CaptureDataEditor.Build.cs
@@ -0,0 +1,30 @@
+// Copyright Epic Games, Inc. All Rights Reserved.
+
+using UnrealBuildTool;
+using System.IO;
+
+public class CaptureDataEditor : ModuleRules
+{
+	public CaptureDataEditor(ReadOnlyTargetRules Target) : base(Target)
+	{
+		PCHUsage = PCHUsageMode.UseExplicitOrSharedPCHs;
+
+		PublicDependencyModuleNames.AddRange(new string[]
+		{
+			
+		});
+
+		PrivateDependencyModuleNames.AddRange(new string[]
+		{
+			"Core",
+			"CoreUObject",
+			"AssetDefinition",
+			"UnrealEd",
+			"CaptureDataCore",
+			"PropertyEditor",
+			"SlateCore",
+			"Slate",
+			"InputCore",
+		});
+	}
+}
diff --git a/Engine/Plugins/VirtualProduction/CaptureData/Source/CaptureDataEditor/Private/AssetDefinition_CameraCalibration.cpp b/Engine/Plugins/VirtualProduction/CaptureData/Source/CaptureDataEditor/Private/AssetDefinition_CameraCalibration.cpp
new file mode 100644
--- /dev/null
+++ b/Engine/Plugins/VirtualProduction/CaptureData/Source/CaptureDataEditor/Private/AssetDefinition_CameraCalibration.cpp
@@ -0,0 +1,32 @@
+// Copyright Epic Games, Inc. All Rights Reserved.
+
+#include "AssetDefinition_CameraCalibration.h"
+#include "CameraCalibration.h"
+
+FText UAssetDefinition_CameraCalibration::GetAssetDisplayName() const
+{
+	return NSLOCTEXT("CaptureData", "CameraCalibrationAssetName", "Camera Calibration");
+}
+
+FLinearColor UAssetDefinition_CameraCalibration::GetAssetColor() const
+{
+	return FColor::Turquoise;
+}
+
+TSoftClassPtr<UObject> UAssetDefinition_CameraCalibration::GetAssetClass() const
+{
+	return UCameraCalibration::StaticClass();
+}
+
+TConstArrayView<FAssetCategoryPath> UAssetDefinition_CameraCalibration::GetAssetCategories() const
+{
+	static FAssetCategoryPath Path(NSLOCTEXT("CaptureData", "CaptureDataAdvancedAssetCategoryLabel", "Advanced"));
+	static FAssetCategoryPath Categories[] = { Path };
+
+	return Categories;
+}
+
+bool UAssetDefinition_CameraCalibration::CanImport() const
+{
+	return true;
+}
\ No newline at end of file
diff --git a/Engine/Plugins/VirtualProduction/CaptureData/Source/CaptureDataEditor/Private/AssetDefinition_CaptureData.cpp b/Engine/Plugins/VirtualProduction/CaptureData/Source/CaptureDataEditor/Private/AssetDefinition_CaptureData.cpp
new file mode 100644
--- /dev/null
+++ b/Engine/Plugins/VirtualProduction/CaptureData/Source/CaptureDataEditor/Private/AssetDefinition_CaptureData.cpp
@@ -0,0 +1,77 @@
+// Copyright Epic Games, Inc. All Rights Reserved.
+
+#include "AssetDefinition_CaptureData.h"
+#include "CaptureData.h"
+#include "ImageSequencePathChecker.h"
+
+//////////////////////////////////////////////////////////////////////////
+// UAssetDefinition_MeshCaptureData
+
+FText UAssetDefinition_MeshCaptureData::GetAssetDisplayName() const
+{
+	return NSLOCTEXT("CaptureData", "MeshCaptureDataAssetName", "Capture Data (Mesh)");
+}
+
+FLinearColor UAssetDefinition_MeshCaptureData::GetAssetColor() const
+{
+	return FColor::Red;
+}
+
+TSoftClassPtr<UObject> UAssetDefinition_MeshCaptureData::GetAssetClass() const
+{
+	return UMeshCaptureData::StaticClass();
+}
+
+TConstArrayView<FAssetCategoryPath> UAssetDefinition_MeshCaptureData::GetAssetCategories() const
+{
+	static FAssetCategoryPath Path(NSLOCTEXT("CaptureData", "CaptureDataAssetCategoryLabel", "MetaHuman Animator"));
+	static FAssetCategoryPath Categories[] = { Path };
+
+	return Categories;
+}
+
+//////////////////////////////////////////////////////////////////////////
+// UAssetDefinition_FootageCaptureData
+
+FText UAssetDefinition_FootageCaptureData::GetAssetDisplayName() const
+{
+	return NSLOCTEXT("CaptureData", "FootageCaptureDataAssetName", "Capture Data (Footage)");
+}
+
+FLinearColor UAssetDefinition_FootageCaptureData::GetAssetColor() const
+{
+	return FColor::Red;
+}
+
+TSoftClassPtr<UObject> UAssetDefinition_FootageCaptureData::GetAssetClass() const
+{
+	return UFootageCaptureData::StaticClass();
+}
+
+TConstArrayView<FAssetCategoryPath> UAssetDefinition_FootageCaptureData::GetAssetCategories() const
+{
+	static FAssetCategoryPath Path(NSLOCTEXT("CaptureData", "CaptureDataAssetCategoryLabel", "MetaHuman Animator"));
+	static FAssetCategoryPath Categories[] = { Path };
+
+	return Categories;
+}
+
+EAssetCommandResult UAssetDefinition_FootageCaptureData::OpenAssets(const FAssetOpenArgs& InOpenArgs) const
+{
+	UE::CaptureData::FImageSequencePathChecker ImageSequencePathChecker(GetAssetDisplayName());
+
+	for (const UFootageCaptureData* FootageCaptureData : InOpenArgs.LoadObjects<UFootageCaptureData>())
+	{
+		if (FootageCaptureData)
+		{
+			ImageSequencePathChecker.Check(*FootageCaptureData);
+		}
+	}
+
+	if (ImageSequencePathChecker.HasError())
+	{
+		ImageSequencePathChecker.DisplayDialog();
+	}
+
+	return UAssetDefinitionDefault::OpenAssets(InOpenArgs);
+}
diff --git a/Engine/Plugins/VirtualProduction/CaptureData/Source/CaptureDataEditor/Private/CameraCalibrationFactory.cpp b/Engine/Plugins/VirtualProduction/CaptureData/Source/CaptureDataEditor/Private/CameraCalibrationFactory.cpp
new file mode 100644
--- /dev/null
+++ b/Engine/Plugins/VirtualProduction/CaptureData/Source/CaptureDataEditor/Private/CameraCalibrationFactory.cpp
@@ -0,0 +1,21 @@
+// Copyright Epic Games, Inc. All Rights Reserved.
+
+#include "CameraCalibrationFactory.h"
+#include "CameraCalibration.h"
+
+
+/* UCameraCalibrationFactory
+ *******************************/
+
+UCameraCalibrationFactory::UCameraCalibrationFactory()
+{
+	bCreateNew = true;
+	bEditAfterNew = true;
+	SupportedClass = UCameraCalibration::StaticClass();
+}
+
+
+UObject* UCameraCalibrationFactory::FactoryCreateNew(UClass* InClass, UObject* InParent, FName InName, EObjectFlags InFlags, UObject* InContext, FFeedbackContext* InWarn)
+{
+	return NewObject<UCameraCalibration>(InParent, InClass, InName, InFlags);
+};
\ No newline at end of file
diff --git a/Engine/Plugins/VirtualProduction/CaptureData/Source/CaptureDataEditor/Private/CaptureDataEditorModule.cpp b/Engine/Plugins/VirtualProduction/CaptureData/Source/CaptureDataEditor/Private/CaptureDataEditorModule.cpp
new file mode 100644
--- /dev/null
+++ b/Engine/Plugins/VirtualProduction/CaptureData/Source/CaptureDataEditor/Private/CaptureDataEditorModule.cpp
@@ -0,0 +1,17 @@
+// Copyright Epic Games, Inc. All Rights Reserved.
+
+#include "CaptureDataEditorModule.h"
+
+#include "Modules/ModuleManager.h"
+
+void FCaptureDataEditorModule::StartupModule()
+{
+
+}
+
+void FCaptureDataEditorModule::ShutdownModule()
+{
+
+}
+
+IMPLEMENT_MODULE(FCaptureDataEditorModule, CaptureDataEditor)
\ No newline at end of file
diff --git a/Engine/Plugins/VirtualProduction/CaptureData/Source/CaptureDataEditor/Private/CaptureDataEditorModule.h b/Engine/Plugins/VirtualProduction/CaptureData/Source/CaptureDataEditor/Private/CaptureDataEditorModule.h
new file mode 100644
--- /dev/null
+++ b/Engine/Plugins/VirtualProduction/CaptureData/Source/CaptureDataEditor/Private/CaptureDataEditorModule.h
@@ -0,0 +1,17 @@
+// Copyright Epic Games, Inc. All Rights Reserved.
+
+#pragma once
+
+#include "Modules/ModuleInterface.h"
+
+class FCaptureDataEditorModule : public IModuleInterface
+{
+public:
+
+	//~ Begin IModuleInterface interface
+	virtual void StartupModule() override;
+	virtual void ShutdownModule() override;
+	//~ End IModuleInterface interface
+
+private:
+};
diff --git a/Engine/Plugins/VirtualProduction/CaptureData/Source/CaptureDataEditor/Private/CaptureDataFactory.cpp b/Engine/Plugins/VirtualProduction/CaptureData/Source/CaptureDataEditor/Private/CaptureDataFactory.cpp
new file mode 100644
--- /dev/null
+++ b/Engine/Plugins/VirtualProduction/CaptureData/Source/CaptureDataEditor/Private/CaptureDataFactory.cpp
@@ -0,0 +1,36 @@
+// Copyright Epic Games, Inc. All Rights Reserved.
+
+#include "CaptureDataFactory.h"
+#include "CaptureData.h"
+
+#include UE_INLINE_GENERATED_CPP_BY_NAME(CaptureDataFactory)
+
+//////////////////////////////////////////////////////////////////////////
+// UMeshCaptureDataFactory
+
+UMeshCaptureDataFactory::UMeshCaptureDataFactory()
+{
+	bCreateNew = true;
+	bEditAfterNew = true;
+	SupportedClass = UMeshCaptureData::StaticClass();
+}
+
+UObject* UMeshCaptureDataFactory::FactoryCreateNew(UClass* InClass, UObject* InParent, FName InName, EObjectFlags InFlags, UObject* InContext, FFeedbackContext* InWarn)
+{
+	return NewObject<UMeshCaptureData>(InParent, InClass, InName, InFlags);
+}
+
+//////////////////////////////////////////////////////////////////////////
+// UFootageCaptureDataFactory
+
+UFootageCaptureDataFactory::UFootageCaptureDataFactory()
+{
+	bCreateNew = true;
+	bEditAfterNew = true;
+	SupportedClass = UFootageCaptureData::StaticClass();
+}
+
+UObject* UFootageCaptureDataFactory::FactoryCreateNew(UClass* InClass, UObject* InParent, FName InName, EObjectFlags InFlags, UObject* InContext, FFeedbackContext* InWarn)
+{
+	return NewObject<UFootageCaptureData>(InParent, InClass, InName, InFlags);
+}
diff --git a/Engine/Plugins/VirtualProduction/CaptureData/Source/CaptureDataEditor/Private/FrameRangeArrayBuilder.cpp b/Engine/Plugins/VirtualProduction/CaptureData/Source/CaptureDataEditor/Private/FrameRangeArrayBuilder.cpp
new file mode 100644
--- /dev/null
+++ b/Engine/Plugins/VirtualProduction/CaptureData/Source/CaptureDataEditor/Private/FrameRangeArrayBuilder.cpp
@@ -0,0 +1,215 @@
+// Copyright Epic Games, Inc. All Rights Reserved.
+
+#include "FrameRangeArrayBuilder.h"
+#include "DetailLayoutBuilder.h"
+#include "IDetailChildrenBuilder.h"
+#include "Widgets/Text/STextBlock.h"
+#include "Widgets/Input/SSpinBox.h"
+#include "Fonts/FontMeasure.h"
+#include "ScopedTransaction.h"
+#include "Widgets/Input/SButton.h"
+#include "Widgets/Input/SEditableTextBox.h"
+
+#define LOCTEXT_NAMESPACE "FrameRange"
+
+FFrameRangeArrayBuilder::FFrameRangeArrayBuilder(TSharedRef<IPropertyHandle> InBaseProperty, TArray<FFrameRange>& InOutFrameRange, FOnGetCurrentFrame* InOnGetCurrentFrameDelegate)
+	: FDetailArrayBuilder(InBaseProperty)
+	, FrameRange(InOutFrameRange)
+	, OnGetCurrentFrameDelegate(InOnGetCurrentFrameDelegate)
+{
+}
+
+void FFrameRangeArrayBuilder::GenerateChildContent(IDetailChildrenBuilder& InOutChildrenBuilder)
+{
+	const FVector2f LabelSize = FSlateApplication::Get().GetRenderer()->GetFontMeasureService()->Measure(TEXT("Start"), IDetailLayoutBuilder::GetDetailFont()) * 1.5;
+
+	for (int32 Index = 0; Index < FrameRange.Num(); ++Index)
+	{
+		TSharedRef<SHorizontalBox> StartRow = SNew(SHorizontalBox);
+
+		StartRow->AddSlot()
+			.AutoWidth()
+			[
+				SNew(STextBlock)
+				.MinDesiredWidth(LabelSize.X)
+				.Text(LOCTEXT("Start", "Start"))
+			];
+
+		StartRow->AddSlot()
+			[
+				SNew(SSpinBox<int32>)
+				.IsEnabled_Lambda([this]()
+				{
+					return GetPropertyHandle()->IsEditable();
+				})
+				.Value_Lambda([this, Index]()
+				{
+					return Index < FrameRange.Num() ? FrameRange[Index].StartFrame : -2; // Index can be temporarily invalid when deleting an entry, return an arbitrary value - should never be actually visible
+				})
+				.OnValueChanged_Lambda([this, Index](int32 InValue)
+				{
+					const FScopedTransaction Transaction{ LOCTEXT("SetStartFrame", "Set Start Frame") };
+
+					GetPropertyHandle()->NotifyPreChange();
+					if (FrameRange[Index].StartFrame == -1 && FrameRange[Index].EndFrame == -1)
+					{
+						FrameRange[Index].EndFrame = InValue;
+					}
+					FrameRange[Index].StartFrame = InValue;
+					GetPropertyHandle()->NotifyPostChange(EPropertyChangeType::ValueSet);
+				})
+			];
+
+		if (OnGetCurrentFrameDelegate)
+		{
+			StartRow->AddSlot()
+				.AutoWidth()
+				[
+					SNew(SButton)
+					.Text(LOCTEXT("Set", "Set"))
+					.HAlign(HAlign_Center)
+					.OnClicked_Lambda([this, Index]()
+					{
+						if (OnGetCurrentFrameDelegate->IsBound())
+						{
+							const FScopedTransaction Transaction{ LOCTEXT("SetStartFrame", "Set Start Frame") };
+
+							GetPropertyHandle()->NotifyPreChange();
+							if (FrameRange[Index].StartFrame == -1 && FrameRange[Index].EndFrame == -1)
+							{
+								FrameRange[Index].EndFrame = OnGetCurrentFrameDelegate->Execute().Value;
+							}
+							FrameRange[Index].StartFrame = OnGetCurrentFrameDelegate->Execute().Value;
+							GetPropertyHandle()->NotifyPostChange(EPropertyChangeType::ValueSet);
+						}
+
+						return FReply::Handled();
+					})
+				];
+		}
+
+
+		TSharedRef<SHorizontalBox> EndRow = SNew(SHorizontalBox);
+
+		EndRow->AddSlot()
+			.AutoWidth()
+			[
+				SNew(STextBlock)
+				.MinDesiredWidth(LabelSize.X)
+				.Text(LOCTEXT("End", "End"))
+			];
+
+		EndRow->AddSlot()
+			[
+				SNew(SSpinBox<int32>)
+				.IsEnabled_Lambda([this]()
+				{
+					return GetPropertyHandle()->IsEditable();
+				})
+				.Value_Lambda([this, Index]()
+				{
+					return Index < FrameRange.Num() ? FrameRange[Index].EndFrame : -2; // Index can be temporarily invalid when deleting an entry, return an arbitrary value - should never be actually visible
+				})
+				.OnValueChanged_Lambda([this, Index](int32 InValue)
+				{
+					const FScopedTransaction Transaction{ LOCTEXT("SetEndFrame", "Set End Frame") };
+
+					GetPropertyHandle()->NotifyPreChange();
+					if (FrameRange[Index].StartFrame == -1 && FrameRange[Index].EndFrame == -1)
+					{
+						FrameRange[Index].StartFrame = InValue;
+					}
+					FrameRange[Index].EndFrame = InValue;
+					GetPropertyHandle()->NotifyPostChange(EPropertyChangeType::ValueSet);
+				})
+			];
+
+		if (OnGetCurrentFrameDelegate)
+		{
+			EndRow->AddSlot()
+				.AutoWidth()
+				[
+					SNew(SButton)
+					.Text(LOCTEXT("Set", "Set"))
+					.HAlign(HAlign_Center)
+					.OnClicked_Lambda([this, Index]()
+					{
+						if (OnGetCurrentFrameDelegate->IsBound())
+						{
+							const FScopedTransaction Transaction{ LOCTEXT("SetEndFrame", "Set End Frame") };
+
+							GetPropertyHandle()->NotifyPreChange();
+							if (FrameRange[Index].StartFrame == -1 && FrameRange[Index].EndFrame == -1)
+							{
+								FrameRange[Index].StartFrame = OnGetCurrentFrameDelegate->Execute().Value;
+							}
+							FrameRange[Index].EndFrame = OnGetCurrentFrameDelegate->Execute().Value;
+							GetPropertyHandle()->NotifyPostChange(EPropertyChangeType::ValueSet);
+						}
+
+						return FReply::Handled();
+					})
+				];
+		}
+
+
+		TSharedRef<SVerticalBox> FrameControls = SNew(SVerticalBox);
+
+		FrameControls->AddSlot()
+			[
+				StartRow
+			];
+
+		FrameControls->AddSlot()
+			[
+				EndRow
+			];
+
+		if (GetPropertyHandle()->IsEditable())
+		{
+			FrameControls->AddSlot()
+				[
+					SNew(SButton)
+					.Text(LOCTEXT("Delete", "Delete"))
+					.HAlign(HAlign_Center)
+					.OnClicked_Lambda([this, Index]()
+					{
+						const FScopedTransaction Transaction{ LOCTEXT("DeleteFrameRange", "Delete Frame Range") };
+
+						GetPropertyHandle()->NotifyPreChange();
+						FrameRange.RemoveAt(Index);
+						GetPropertyHandle()->NotifyPostChange(EPropertyChangeType::ArrayRemove);
+
+						return FReply::Handled();
+					})
+				];
+		}
+
+
+		InOutChildrenBuilder.AddCustomRow(LOCTEXT("Frame", "Frame"))
+			.NameContent()
+			[
+				SNew(SEditableTextBox)
+				.Text_Lambda([this, Index]()
+				{
+					return Index < FrameRange.Num() ? FText::FromString(FrameRange[Index].Name) : FText(); // Index can be temporarily invalid when deleting an entry, return an arbitrary value - should never be actually visible
+				})
+				.OnTextCommitted_Lambda([this, Index](FText InValue, ETextCommit::Type)
+				{
+					const FScopedTransaction Transaction{ LOCTEXT("NameFrameRange", "Name Frame Range") };
+
+					GetPropertyHandle()->NotifyPreChange();
+					FrameRange[Index].Name = InValue.ToString();
+					GetPropertyHandle()->NotifyPostChange(EPropertyChangeType::ValueSet);						
+				})
+			]
+			.ValueContent()
+			.MinDesiredWidth(450.0f)
+			.MaxDesiredWidth(0.0f)
+			[
+				FrameControls
+			];
+	}
+}
+
+#undef LOCTEXT_NAMESPACE
diff --git a/Engine/Plugins/VirtualProduction/CaptureData/Source/CaptureDataEditor/Public/AssetDefinition_CameraCalibration.h b/Engine/Plugins/VirtualProduction/CaptureData/Source/CaptureDataEditor/Public/AssetDefinition_CameraCalibration.h
new file mode 100644
--- /dev/null
+++ b/Engine/Plugins/VirtualProduction/CaptureData/Source/CaptureDataEditor/Public/AssetDefinition_CameraCalibration.h
@@ -0,0 +1,23 @@
+// Copyright Epic Games, Inc. All Rights Reserved.
+
+#pragma once
+
+#include "AssetDefinitionDefault.h"
+
+#include "AssetDefinition_CameraCalibration.generated.h"
+
+UCLASS()
+class UAssetDefinition_CameraCalibration
+	: public UAssetDefinitionDefault
+{
+	GENERATED_BODY()
+
+public:
+	//~Begin UAssetDefinitionDefault interface
+	virtual FText GetAssetDisplayName() const override;
+	virtual FLinearColor GetAssetColor() const override;
+	virtual TSoftClassPtr<UObject> GetAssetClass() const override;
+	virtual TConstArrayView<FAssetCategoryPath> GetAssetCategories() const override;
+	virtual bool CanImport() const override;
+	//~End UAssetDefinitionDefault interface
+};
\ No newline at end of file
diff --git a/Engine/Plugins/VirtualProduction/CaptureData/Source/CaptureDataEditor/Public/AssetDefinition_CaptureData.h b/Engine/Plugins/VirtualProduction/CaptureData/Source/CaptureDataEditor/Public/AssetDefinition_CaptureData.h
new file mode 100644
--- /dev/null
+++ b/Engine/Plugins/VirtualProduction/CaptureData/Source/CaptureDataEditor/Public/AssetDefinition_CaptureData.h
@@ -0,0 +1,46 @@
+// Copyright Epic Games, Inc. All Rights Reserved.
+
+#pragma once
+
+#include "AssetDefinitionDefault.h"
+
+#include "AssetDefinition_CaptureData.generated.h"
+
+//////////////////////////////////////////////////////////////////////////
+// UAssetDefinition_MeshCaptureData
+
+UCLASS()
+class UAssetDefinition_MeshCaptureData
+	: public UAssetDefinitionDefault
+{
+	GENERATED_BODY()
+
+public:
+	//~Begin UAssetDefinitionDefault interface
+	virtual FText GetAssetDisplayName() const override;
+	virtual FLinearColor GetAssetColor() const override;
+	virtual TSoftClassPtr<UObject> GetAssetClass() const override;
+	virtual TConstArrayView<FAssetCategoryPath> GetAssetCategories() const override;
+	//~End UAssetDefinitionDefault interface
+};
+
+//////////////////////////////////////////////////////////////////////////
+// UAssetDefinition_FootageCaptureData
+
+UCLASS()
+class UAssetDefinition_FootageCaptureData
+	: public UAssetDefinitionDefault
+{
+	GENERATED_BODY()
+
+public:
+	//~Begin UAssetDefinitionDefault interface
+	virtual FText GetAssetDisplayName() const override;
+	virtual FLinearColor GetAssetColor() const override;
+	virtual TSoftClassPtr<UObject> GetAssetClass() const override;
+	virtual TConstArrayView<FAssetCategoryPath> GetAssetCategories() const override;
+
+	virtual EAssetCommandResult OpenAssets(const FAssetOpenArgs& OpenArgs) const override;
+
+	//~End UAssetDefinitionDefault interface
+};
\ No newline at end of file
diff --git a/Engine/Plugins/VirtualProduction/CaptureData/Source/CaptureDataEditor/Public/CameraCalibrationFactory.h b/Engine/Plugins/VirtualProduction/CaptureData/Source/CaptureDataEditor/Public/CameraCalibrationFactory.h
new file mode 100644
--- /dev/null
+++ b/Engine/Plugins/VirtualProduction/CaptureData/Source/CaptureDataEditor/Public/CameraCalibrationFactory.h
@@ -0,0 +1,23 @@
+// Copyright Epic Games, Inc. All Rights Reserved.
+
+#pragma once
+
+#include "Factories/Factory.h"
+
+#include "CameraCalibrationFactory.generated.h"
+
+/**
+ * Implements a factory for UCameraCalibrationFactory objects from file
+ */
+UCLASS(hidecategories=Object)
+class UCameraCalibrationFactory
+	: public UFactory
+{
+	GENERATED_BODY()
+
+public:
+	UCameraCalibrationFactory();
+
+	//~ UFactory interface
+	virtual UObject* FactoryCreateNew(UClass* InClass, UObject* InParent, FName InName, EObjectFlags InFlags, UObject* InContext, FFeedbackContext* InWarn) override;
+};
diff --git a/Engine/Plugins/VirtualProduction/CaptureData/Source/CaptureDataEditor/Public/CaptureDataFactory.h b/Engine/Plugins/VirtualProduction/CaptureData/Source/CaptureDataEditor/Public/CaptureDataFactory.h
new file mode 100644
--- /dev/null
+++ b/Engine/Plugins/VirtualProduction/CaptureData/Source/CaptureDataEditor/Public/CaptureDataFactory.h
@@ -0,0 +1,40 @@
+// Copyright Epic Games, Inc. All Rights Reserved.
+
+#pragma once
+
+#include "Factories/Factory.h"
+
+#include "CaptureDataFactory.generated.h"
+
+
+//////////////////////////////////////////////////////////////////////////
+// UMeshCaptureDataFactory
+
+UCLASS(hidecategories=Object)
+class UMeshCaptureDataFactory
+	: public UFactory
+{
+	GENERATED_BODY()
+
+public:
+	//~ UFactory Interface
+	UMeshCaptureDataFactory();
+
+	virtual UObject* FactoryCreateNew(UClass* InClass, UObject* InParent, FName InName, EObjectFlags InFlags, UObject* InContext, FFeedbackContext* InWarn) override;
+};
+
+//////////////////////////////////////////////////////////////////////////
+// UFootageCaptureDataFactory
+
+UCLASS(hideCategories=Object)
+class UFootageCaptureDataFactory
+	: public UFactory
+{
+	GENERATED_BODY()
+
+public:
+	//~ UFactory Interface
+	UFootageCaptureDataFactory();
+
+	virtual UObject* FactoryCreateNew(UClass* InClass, UObject* InParent, FName InName, EObjectFlags InFlags, UObject* InContext, FFeedbackContext* InWarn) override;
+};
\ No newline at end of file
diff --git a/Engine/Plugins/VirtualProduction/CaptureData/Source/CaptureDataEditor/Public/Customizations/CaptureDataCustomizations.cpp b/Engine/Plugins/VirtualProduction/CaptureData/Source/CaptureDataEditor/Public/Customizations/CaptureDataCustomizations.cpp
new file mode 100644
--- /dev/null
+++ b/Engine/Plugins/VirtualProduction/CaptureData/Source/CaptureDataEditor/Public/Customizations/CaptureDataCustomizations.cpp
@@ -0,0 +1,39 @@
+// Copyright Epic Games, Inc.All Rights Reserved.
+
+#include "Customizations/CaptureDataCustomizations.h"
+#include "CaptureData.h"
+#include "FrameRangeArrayBuilder.h"
+#include "DetailLayoutBuilder.h"
+#include "DetailCategoryBuilder.h"
+
+TSharedRef<IDetailCustomization> FFootageCaptureDataCustomization::MakeInstance()
+{
+	return MakeShared<FFootageCaptureDataCustomization>();
+}
+
+void FFootageCaptureDataCustomization::CustomizeDetails(IDetailLayoutBuilder& InDetailBuilder)
+{
+	InDetailBuilder.AddPropertyToCategory(InDetailBuilder.GetProperty(GET_MEMBER_NAME_CHECKED(UFootageCaptureData, Views)));
+	InDetailBuilder.AddPropertyToCategory(InDetailBuilder.GetProperty(GET_MEMBER_NAME_CHECKED(UFootageCaptureData, Audios)));
+	InDetailBuilder.AddPropertyToCategory(InDetailBuilder.GetProperty(GET_MEMBER_NAME_CHECKED(UFootageCaptureData, bAudioTimecodePresent)));
+	InDetailBuilder.AddPropertyToCategory(InDetailBuilder.GetProperty(GET_MEMBER_NAME_CHECKED(UFootageCaptureData, AudioTimecode)));
+	InDetailBuilder.AddPropertyToCategory(InDetailBuilder.GetProperty(GET_MEMBER_NAME_CHECKED(UFootageCaptureData, AudioTimecodeRate)));
+	InDetailBuilder.AddPropertyToCategory(InDetailBuilder.GetProperty(GET_MEMBER_NAME_CHECKED(UFootageCaptureData, CameraCalibrations)));
+	InDetailBuilder.AddPropertyToCategory(InDetailBuilder.GetProperty(GET_MEMBER_NAME_CHECKED(UFootageCaptureData, Metadata)));
+
+	UFootageCaptureData* FootageCaptureData = nullptr;
+
+	// Get the performance object that we're building the details panel for.
+	if (!InDetailBuilder.GetSelectedObjects().IsEmpty())
+	{
+		FootageCaptureData = Cast<UFootageCaptureData>(InDetailBuilder.GetSelectedObjects()[0].Get());
+	}
+	else
+	{
+		return;
+	}
+
+	TSharedRef<IPropertyHandle> CaptureExcludedFramesProperty = InDetailBuilder.GetProperty(GET_MEMBER_NAME_CHECKED(UFootageCaptureData, CaptureExcludedFrames));
+	IDetailCategoryBuilder& CaptureExcludedFramesCategory = InDetailBuilder.EditCategory(CaptureExcludedFramesProperty->GetDefaultCategoryName());
+	CaptureExcludedFramesCategory.AddCustomBuilder(MakeShareable(new FFrameRangeArrayBuilder(CaptureExcludedFramesProperty, FootageCaptureData->CaptureExcludedFrames)), false);
+}
diff --git a/Engine/Plugins/VirtualProduction/CaptureData/Source/CaptureDataEditor/Public/Customizations/CaptureDataCustomizations.h b/Engine/Plugins/VirtualProduction/CaptureData/Source/CaptureDataEditor/Public/Customizations/CaptureDataCustomizations.h
new file mode 100644
--- /dev/null
+++ b/Engine/Plugins/VirtualProduction/CaptureData/Source/CaptureDataEditor/Public/Customizations/CaptureDataCustomizations.h
@@ -0,0 +1,18 @@
+// Copyright Epic Games, Inc. All Rights Reserved.
+
+#pragma once
+
+#include "CoreMinimal.h"
+#include "UObject/Object.h"
+#include "IDetailCustomization.h"
+
+class CAPTUREDATAEDITOR_API FFootageCaptureDataCustomization
+	: public IDetailCustomization
+{
+public:
+	static TSharedRef<IDetailCustomization> MakeInstance();
+
+	//~Begin IDetailCustomization interface
+	virtual void CustomizeDetails(IDetailLayoutBuilder& InDetailBuilder) override;
+	//~End IDetailCustomization interface
+};
\ No newline at end of file
diff --git a/Engine/Plugins/VirtualProduction/CaptureData/Source/CaptureDataEditor/Public/FrameRangeArrayBuilder.h b/Engine/Plugins/VirtualProduction/CaptureData/Source/CaptureDataEditor/Public/FrameRangeArrayBuilder.h
new file mode 100644
--- /dev/null
+++ b/Engine/Plugins/VirtualProduction/CaptureData/Source/CaptureDataEditor/Public/FrameRangeArrayBuilder.h
@@ -0,0 +1,22 @@
+// Copyright Epic Games, Inc. All Rights Reserved.
+
+#pragma once
+
+#include "FrameRange.h"
+#include "PropertyCustomizationHelpers.h"
+
+class CAPTUREDATAEDITOR_API FFrameRangeArrayBuilder : public FDetailArrayBuilder
+{
+public:
+
+	DECLARE_DELEGATE_RetVal(FFrameNumber, FOnGetCurrentFrame);
+
+	FFrameRangeArrayBuilder(TSharedRef<IPropertyHandle> InBaseProperty, TArray<FFrameRange>& InOutFrameRange, FOnGetCurrentFrame* InOnGetCurrentFrameDelegate = nullptr);
+
+	virtual void GenerateChildContent(IDetailChildrenBuilder& InOutChildrenBuilder) override;
+
+private:
+
+	TArray<FFrameRange>& FrameRange;
+	FOnGetCurrentFrame* OnGetCurrentFrameDelegate = nullptr;
+};
