Deprecate Views; make image and depth sequences

#virtualized

diff --git a/Engine/Plugins/VirtualProduction/CaptureData/CaptureData.uplugin b/Engine/Plugins/VirtualProduction/CaptureData/CaptureData.uplugin
--- a/Engine/Plugins/VirtualProduction/CaptureData/CaptureData.uplugin
+++ b/Engine/Plugins/VirtualProduction/CaptureData/CaptureData.uplugin
@@ -10,7 +10,7 @@
 	"DocsURL": "",
 	"MarketplaceURL": "",
 	"SupportURL": "",
-	"CanContainContent": false,
+	"CanContainContent": true,
 	"Hidden": true,
 	"IsBetaVersion": false,
 	"Installed": false,
@@ -24,6 +24,11 @@
 			"Name": "CaptureDataEditor",
 			"Type": "Editor",
 			"LoadingPhase": "Default"
+		},
+		{
+			"Name": "CaptureDataUtils",
+			"Type": "Runtime",
+			"LoadingPhase": "Default"
 		}
 	],
 	"Plugins": [
diff --git a/Engine/Plugins/VirtualProduction/CaptureData/Source/CaptureDataCore/CaptureDataCore.Build.cs b/Engine/Plugins/VirtualProduction/CaptureData/Source/CaptureDataCore/CaptureDataCore.Build.cs
--- a/Engine/Plugins/VirtualProduction/CaptureData/Source/CaptureDataCore/CaptureDataCore.Build.cs
+++ b/Engine/Plugins/VirtualProduction/CaptureData/Source/CaptureDataCore/CaptureDataCore.Build.cs
@@ -21,7 +21,8 @@ public class CaptureDataCore : ModuleRules
 			"CameraCalibrationCore",
 			"MeshDescription",
 			"StaticMeshDescription",
-			"Engine"
+			"Engine",
+			"CaptureDataUtils"
 		});
 	}
 }
diff --git a/Engine/Plugins/VirtualProduction/CaptureData/Source/CaptureDataCore/Private/CaptureData.cpp b/Engine/Plugins/VirtualProduction/CaptureData/Source/CaptureDataCore/Private/CaptureData.cpp
--- a/Engine/Plugins/VirtualProduction/CaptureData/Source/CaptureDataCore/Private/CaptureData.cpp
+++ b/Engine/Plugins/VirtualProduction/CaptureData/Source/CaptureDataCore/Private/CaptureData.cpp
@@ -22,6 +22,7 @@
 #include "Engine/SkeletalMesh.h"
 
 #if WITH_EDITOR
+#include "ImageSequenceTimecodeUtils.h"
 #include "IDirectoryWatcher.h"
 #include "DirectoryWatcherModule.h"
 #endif
@@ -252,10 +253,8 @@ TArray<UFootageCaptureData::FPathAssociation> UFootageCaptureData::CheckImageSeq
 {
 	TArray<FPathAssociation> InvalidImageSequences;
 
-	for (const FFootageCaptureView& View : Views)
+	for (const TObjectPtr<class UImgMediaSource>& ImageSequence : ImageSequences)
 	{
-		const UImgMediaSource* ImageSequence = View.ImageSequence;
-
 		if (ImageSequence && !FPaths::DirectoryExists(ImageSequence->GetFullPath()))
 		{
 			InvalidImageSequences.Emplace(ImageSequence->GetFullPath(), ImageSequence->GetPathName());
@@ -267,23 +266,27 @@ TArray<UFootageCaptureData::FPathAssociation> UFootageCaptureData::CheckImageSeq
 
 UFootageCaptureData::FVerifyResult UFootageCaptureData::ViewsContainsValidData() const
 {
-	CHECK_AND_PRINT_RET(!Views.IsEmpty(), TEXT("Capture Data doesn't contain views"));
+	CHECK_AND_PRINT_RET(!ImageSequences.IsEmpty(), TEXT("Capture Data doesn't contain image sequences"));
+	CHECK_AND_PRINT_RET(!DepthSequences.IsEmpty(), TEXT("Capture Data doesn't contain depth sequences"));
 
-	for (const FFootageCaptureView& View : Views)
+	for (int32 Index = 0; Index < ImageSequences.Num(); Index++)
 	{
 		int32 NumImageFrames = 0;
 		FIntVector2 ImageDimensions;
 
-		CHECK_AND_PRINT_RET(FImageSequenceUtils::GetImageSequenceInfo(View.ImageSequence, ImageDimensions, NumImageFrames), TEXT("Image Sequence asset is invalid"));
+		UImgMediaSource* ImageSequence = ImageSequences[Index];
+		UImgMediaSource* DepthSequence = DepthSequences[Index];
+
+		CHECK_AND_PRINT_RET(FImageSequenceUtils::GetImageSequenceInfo(ImageSequence, ImageDimensions, NumImageFrames), TEXT("Image Sequence asset is invalid"));
 
 		int32 NumDepthFrames = 0;
 		FIntVector2 DepthDimensions;
 
-		CHECK_AND_PRINT_RET(FImageSequenceUtils::GetImageSequenceInfo(View.DepthSequence, DepthDimensions, NumImageFrames), TEXT("Depth Sequence asset is invalid"));
+		CHECK_AND_PRINT_RET(FImageSequenceUtils::GetImageSequenceInfo(DepthSequence, DepthDimensions, NumImageFrames), TEXT("Depth Sequence asset is invalid"));
 
-		CHECK_AND_PRINT_RET(View.ImageSequence->FrameRateOverride.IsValid(), TEXT("Image Sequence asset contains invalid frame rate"));
-		CHECK_AND_PRINT_RET(View.DepthSequence->FrameRateOverride.IsValid(), TEXT("Depth Sequence asset contains invalid frame rate"));
-		CHECK_AND_PRINT_RET(View.ImageSequence->FrameRateOverride == View.DepthSequence->FrameRateOverride, TEXT("Image Sequence asset and Depth Sequence asset contain different frame rates"));
+		CHECK_AND_PRINT_RET(ImageSequence->FrameRateOverride.IsValid(), TEXT("Image Sequence asset contains invalid frame rate"));
+		CHECK_AND_PRINT_RET(DepthSequence->FrameRateOverride.IsValid(), TEXT("Depth Sequence asset contains invalid frame rate"));
+		CHECK_AND_PRINT_RET(ImageSequence->FrameRateOverride == DepthSequence->FrameRateOverride, TEXT("Image Sequence asset and Depth Sequence asset contain different frame rates"));
 	}
 
 	return MakeValue();
@@ -360,10 +363,14 @@ void UFootageCaptureData::GetFrameRanges(const FFrameRate& InTargetRate, ETimeco
 	OutProcessingFrameRange = TRange<FFrameNumber>(0, 0);
 	OutMaximumFrameRange = TRange<FFrameNumber>(0, 0);
 
-	for (int32 Index = 0; Index < Views.Num(); ++Index)
+	for (int32 Index = 0; Index < ImageSequences.Num(); ++Index)
+	{
+		OutMediaFrameRanges.Add(ImageSequences[Index], GetFrameRange(InTargetRate, ImageSequences[Index], GetEffectiveImageTimecode(Index), GetEffectiveImageTimecodeRate(Index), InTimecodeAlignment == ETimecodeAlignment::None));
+	}
+
+	for (int32 Index = 0; Index < DepthSequences.Num(); ++Index)
 	{
-		OutMediaFrameRanges.Add(Views[Index].ImageSequence, GetFrameRange(InTargetRate, Views[Index].ImageSequence, GetEffectiveImageTimecode(Index), GetEffectiveImageTimecodeRate(Index), InTimecodeAlignment == ETimecodeAlignment::None));
-		OutMediaFrameRanges.Add(Views[Index].DepthSequence, GetFrameRange(InTargetRate, Views[Index].DepthSequence, GetEffectiveDepthTimecode(Index), GetEffectiveDepthTimecodeRate(Index), InTimecodeAlignment == ETimecodeAlignment::None));
+		OutMediaFrameRanges.Add(DepthSequences[Index], GetFrameRange(InTargetRate, DepthSequences[Index], GetEffectiveDepthTimecode(Index), GetEffectiveDepthTimecodeRate(Index), InTimecodeAlignment == ETimecodeAlignment::None));
 	}
 
 	if (bInIncludeAudio)
@@ -449,92 +456,79 @@ TRange<FFrameNumber> UFootageCaptureData::GetAudioFrameRange(const FFrameRate& I
 
 FTimecode UFootageCaptureData::GetEffectiveImageTimecode(int32 InView) const
 {
-	check(InView < Views.Num());
+	check(InView < ImageSequences.Num());
+	check(InView < DepthSequences.Num());
 
-	FTimecode Timecode;
-
-	if (Views[InView].bImageTimecodePresent)
-	{
-		Timecode = Views[InView].ImageTimecode;
-	}
-	else if (Views[InView].bDepthTimecodePresent)
-	{
-		Timecode = Views[InView].DepthTimecode;
-	}
-	else
+#if WITH_EDITOR
+	TOptional<FTimecode> TimecodeOpt = UImageSequenceTimecodeUtils::GetTimecode(ImageSequences[InView]);
+	if (TimecodeOpt.IsSet())
 	{
-		FFrameRate FrameRate;
-		GetDefaultTimecodeInfo(Timecode, FrameRate);
+		return MoveTemp(TimecodeOpt.GetValue());
 	}
+#endif
+
+	FTimecode Timecode;
+	FFrameRate FrameRate;
+	GetDefaultTimecodeInfo(Timecode, FrameRate);
 
 	return Timecode;
 }
 
 FFrameRate UFootageCaptureData::GetEffectiveImageTimecodeRate(int32 InView) const
 {
-	check(InView < Views.Num());
+	check(InView < ImageSequences.Num());
+	check(InView < DepthSequences.Num());
 
-	FFrameRate FrameRate;
-
-	if (Views[InView].bImageTimecodePresent)
-	{
-		FrameRate = Views[InView].ImageTimecodeRate;
-	}
-	else if (Views[InView].bDepthTimecodePresent)
-	{
-		FrameRate = Views[InView].DepthTimecodeRate;
-	}
-	else
+#if WITH_EDITOR
+	TOptional<FFrameRate> FrameRateOpt = UImageSequenceTimecodeUtils::GetFrameRate(ImageSequences[InView]);
+	if (FrameRateOpt.IsSet())
 	{
-		FTimecode Timecode;
-		GetDefaultTimecodeInfo(Timecode, FrameRate);
+		return MoveTemp(FrameRateOpt.GetValue());
 	}
+#endif
+
+	FTimecode Timecode;
+	FFrameRate FrameRate;
+	GetDefaultTimecodeInfo(Timecode, FrameRate);
 
 	return FrameRate;
 }
 
 FTimecode UFootageCaptureData::GetEffectiveDepthTimecode(int32 InView) const
 {
-	check(InView < Views.Num());
-
-	FTimecode Timecode;
+	check(InView < ImageSequences.Num());
+	check(InView < DepthSequences.Num());
 
-	if (Views[InView].bDepthTimecodePresent)
-	{
-		Timecode = Views[InView].DepthTimecode;
-	}
-	else if (Views[InView].bImageTimecodePresent)
-	{
-		Timecode = Views[InView].ImageTimecode;
-	}
-	else
+#if WITH_EDITOR
+	TOptional<FTimecode> TimecodeOpt = UImageSequenceTimecodeUtils::GetTimecode(DepthSequences[InView]);
+	if (TimecodeOpt.IsSet())
 	{
-		FFrameRate FrameRate;
-		GetDefaultTimecodeInfo(Timecode, FrameRate);
+		return MoveTemp(TimecodeOpt.GetValue());
 	}
+#endif
+
+	FTimecode Timecode;
+	FFrameRate FrameRate;
+	GetDefaultTimecodeInfo(Timecode, FrameRate);
 
 	return Timecode;
 }
 
 FFrameRate UFootageCaptureData::GetEffectiveDepthTimecodeRate(int32 InView) const
 {
-	check(InView < Views.Num());
-
-	FFrameRate FrameRate;
+	check(InView < ImageSequences.Num());
+	check(InView < DepthSequences.Num());
 
-	if (Views[InView].bDepthTimecodePresent)
-	{
-		FrameRate = Views[InView].DepthTimecodeRate;
-	}
-	else if (Views[InView].bImageTimecodePresent)
-	{
-		FrameRate = Views[InView].ImageTimecodeRate;
-	}
-	else
+#if WITH_EDITOR
+	TOptional<FFrameRate> FrameRateOpt = UImageSequenceTimecodeUtils::GetFrameRate(DepthSequences[InView]);
+	if (FrameRateOpt.IsSet())
 	{
-		FTimecode Timecode;
-		GetDefaultTimecodeInfo(Timecode, FrameRate);
+		return MoveTemp(FrameRateOpt.GetValue());
 	}
+#endif
+	FTimecode Timecode;
+	FFrameRate FrameRate;
+	GetDefaultTimecodeInfo(Timecode, FrameRate);
 
 	return FrameRate;
 }
@@ -575,25 +569,33 @@ FFrameRate UFootageCaptureData::GetEffectiveAudioTimecodeRate() const
 
 void UFootageCaptureData::GetDefaultTimecodeInfo(FTimecode& OutTimecode, FFrameRate& OutFrameRate) const
 {
-	for (const FFootageCaptureView& View : Views)
+#if WITH_EDITOR
+	for (const TObjectPtr<class UImgMediaSource>& ImageSequence : ImageSequences)
 	{
-		if (View.bImageTimecodePresent)
+		TOptional<FTimecode> TimecodeOpt = UImageSequenceTimecodeUtils::GetTimecode(ImageSequence);
+		TOptional<FFrameRate> FrameRateOpt = UImageSequenceTimecodeUtils::GetFrameRate(ImageSequence);
+
+		if (TimecodeOpt.IsSet() && FrameRateOpt.IsSet())
 		{
-			OutTimecode = View.ImageTimecode;
-			OutFrameRate = View.ImageTimecodeRate;
+			OutTimecode = MoveTemp(TimecodeOpt.GetValue());
+			OutFrameRate = MoveTemp(FrameRateOpt.GetValue());
 			return;
 		}
 	}
 
-	for (const FFootageCaptureView& View : Views)
+	for (const TObjectPtr<class UImgMediaSource>& DepthSequence : DepthSequences)
 	{
-		if (View.bDepthTimecodePresent)
+		TOptional<FTimecode> TimecodeOpt = UImageSequenceTimecodeUtils::GetTimecode(DepthSequence);
+		TOptional<FFrameRate> FrameRateOpt = UImageSequenceTimecodeUtils::GetFrameRate(DepthSequence);
+
+		if (TimecodeOpt.IsSet() && FrameRateOpt.IsSet())
 		{
-			OutTimecode = View.DepthTimecode;
-			OutFrameRate = View.DepthTimecodeRate;
+			OutTimecode = MoveTemp(TimecodeOpt.GetValue());
+			OutFrameRate = MoveTemp(FrameRateOpt.GetValue());
 			return;
 		}
 	}
+#endif
 
 	if (bAudioTimecodePresent)
 	{
@@ -627,6 +629,23 @@ void UFootageCaptureData::PostLoad()
 		Audios.Add(MoveTemp(Audio_DEPRECATED));
 		Audio_DEPRECATED = nullptr;
 	}
+
+#if WITH_EDITOR
+	for (FFootageCaptureView& View : Views_DEPRECATED)
+	{
+		if (View.bImageTimecodePresent)
+		{
+			UImageSequenceTimecodeUtils::SetTimecodeInfo(View.ImageTimecode, View.ImageTimecodeRate, View.ImageSequence.Get());
+		}
+		ImageSequences.Add(MoveTemp(View.ImageSequence));
+
+		if (View.bDepthTimecodePresent)
+		{
+			UImageSequenceTimecodeUtils::SetTimecodeInfo(View.DepthTimecode, View.DepthTimecodeRate, View.DepthSequence.Get());
+		}
+		DepthSequences.Add(MoveTemp(View.DepthSequence));
+	}
+#endif
 }
 
 void UFootageCaptureData::PopulateCameraNames(UFootageCaptureData* InFootageCaptureData, FString& InOutCamera, TArray<TSharedPtr<FString>>& OutCameraNames)
diff --git a/Engine/Plugins/VirtualProduction/CaptureData/Source/CaptureDataCore/Public/CaptureData.h b/Engine/Plugins/VirtualProduction/CaptureData/Source/CaptureDataCore/Public/CaptureData.h
--- a/Engine/Plugins/VirtualProduction/CaptureData/Source/CaptureDataCore/Public/CaptureData.h
+++ b/Engine/Plugins/VirtualProduction/CaptureData/Source/CaptureDataCore/Public/CaptureData.h
@@ -206,8 +206,14 @@ public:
 
 public:
 
+	UPROPERTY(meta = (DeprecatedProperty, DeprecationMessage = "Views are deprecated. Please use Image and Depth sequences instead."))
+	TArray<FFootageCaptureView> Views_DEPRECATED;
+
+	UPROPERTY(EditAnywhere, BlueprintReadWrite, Category = "Capture")
+	TArray<TObjectPtr<class UImgMediaSource>> ImageSequences;
+
 	UPROPERTY(EditAnywhere, BlueprintReadWrite, Category = "Capture")
-	TArray<FFootageCaptureView> Views;
+	TArray<TObjectPtr<class UImgMediaSource>> DepthSequences;
 
 	UPROPERTY(EditAnywhere, BlueprintReadWrite, Category = "Capture")
 	TArray<TObjectPtr<class USoundWave>> Audios;
diff --git a/Engine/Plugins/VirtualProduction/CaptureData/Source/CaptureDataEditor/Private/Customizations/CaptureDataCustomizations.cpp b/Engine/Plugins/VirtualProduction/CaptureData/Source/CaptureDataEditor/Private/Customizations/CaptureDataCustomizations.cpp
new file mode 100644
--- /dev/null
+++ b/Engine/Plugins/VirtualProduction/CaptureData/Source/CaptureDataEditor/Private/Customizations/CaptureDataCustomizations.cpp
@@ -0,0 +1,40 @@
+// Copyright Epic Games, Inc.All Rights Reserved.
+
+#include "Customizations/CaptureDataCustomizations.h"
+#include "CaptureData.h"
+#include "FrameRangeArrayBuilder.h"
+#include "DetailLayoutBuilder.h"
+#include "DetailCategoryBuilder.h"
+
+TSharedRef<IDetailCustomization> FFootageCaptureDataCustomization::MakeInstance()
+{
+	return MakeShared<FFootageCaptureDataCustomization>();
+}
+
+void FFootageCaptureDataCustomization::CustomizeDetails(IDetailLayoutBuilder& InDetailBuilder)
+{
+	InDetailBuilder.AddPropertyToCategory(InDetailBuilder.GetProperty(GET_MEMBER_NAME_CHECKED(UFootageCaptureData, ImageSequences)));
+	InDetailBuilder.AddPropertyToCategory(InDetailBuilder.GetProperty(GET_MEMBER_NAME_CHECKED(UFootageCaptureData, DepthSequences)));
+	InDetailBuilder.AddPropertyToCategory(InDetailBuilder.GetProperty(GET_MEMBER_NAME_CHECKED(UFootageCaptureData, Audios)));
+	InDetailBuilder.AddPropertyToCategory(InDetailBuilder.GetProperty(GET_MEMBER_NAME_CHECKED(UFootageCaptureData, bAudioTimecodePresent)));
+	InDetailBuilder.AddPropertyToCategory(InDetailBuilder.GetProperty(GET_MEMBER_NAME_CHECKED(UFootageCaptureData, AudioTimecode)));
+	InDetailBuilder.AddPropertyToCategory(InDetailBuilder.GetProperty(GET_MEMBER_NAME_CHECKED(UFootageCaptureData, AudioTimecodeRate)));
+	InDetailBuilder.AddPropertyToCategory(InDetailBuilder.GetProperty(GET_MEMBER_NAME_CHECKED(UFootageCaptureData, CameraCalibrations)));
+	InDetailBuilder.AddPropertyToCategory(InDetailBuilder.GetProperty(GET_MEMBER_NAME_CHECKED(UFootageCaptureData, Metadata)));
+
+	UFootageCaptureData* FootageCaptureData = nullptr;
+
+	// Get the performance object that we're building the details panel for.
+	if (!InDetailBuilder.GetSelectedObjects().IsEmpty())
+	{
+		FootageCaptureData = Cast<UFootageCaptureData>(InDetailBuilder.GetSelectedObjects()[0].Get());
+	}
+	else
+	{
+		return;
+	}
+
+	TSharedRef<IPropertyHandle> CaptureExcludedFramesProperty = InDetailBuilder.GetProperty(GET_MEMBER_NAME_CHECKED(UFootageCaptureData, CaptureExcludedFrames));
+	IDetailCategoryBuilder& CaptureExcludedFramesCategory = InDetailBuilder.EditCategory(CaptureExcludedFramesProperty->GetDefaultCategoryName());
+	CaptureExcludedFramesCategory.AddCustomBuilder(MakeShareable(new FFrameRangeArrayBuilder(CaptureExcludedFramesProperty, FootageCaptureData->CaptureExcludedFrames)), false);
+}
diff --git a/Engine/Plugins/VirtualProduction/CaptureData/Source/CaptureDataEditor/Public/Customizations/CaptureDataCustomizations.cpp b/Engine/Plugins/VirtualProduction/CaptureData/Source/CaptureDataEditor/Public/Customizations/CaptureDataCustomizations.cpp
deleted file mode 100644
diff --git a/Engine/Plugins/VirtualProduction/CaptureData/Source/CaptureDataUtils/CaptureDataUtils.Build.cs b/Engine/Plugins/VirtualProduction/CaptureData/Source/CaptureDataUtils/CaptureDataUtils.Build.cs
new file mode 100644
--- /dev/null
+++ b/Engine/Plugins/VirtualProduction/CaptureData/Source/CaptureDataUtils/CaptureDataUtils.Build.cs
@@ -0,0 +1,41 @@
+// Copyright Epic Games, Inc. All Rights Reserved.
+
+using UnrealBuildTool;
+using System.IO;
+using UnrealBuildTool.Rules;
+
+public class CaptureDataUtils : ModuleRules
+{
+	protected bool BuildForDevelopment
+	{
+		get
+		{
+			// Check if source is available
+			string SourceFilesPath = Path.Combine(ModuleDirectory, "Private");
+			return Directory.Exists(SourceFilesPath) &&
+				   Directory.GetFiles(SourceFilesPath).Length > 0;
+		}
+	}
+
+	public CaptureDataUtils(ReadOnlyTargetRules Target) : base(Target)
+	{
+		bUsePrecompiled = !BuildForDevelopment;
+		PCHUsage = PCHUsageMode.UseExplicitOrSharedPCHs;
+
+		PublicDependencyModuleNames.AddRange(new string[]
+		{
+			"Core",
+			"CoreUObject",
+			"Engine",
+			"ImgMedia",
+		});
+
+		if (Target.bBuildEditor == true)
+		{
+			PrivateDependencyModuleNames.AddRange(new string[]
+			{
+				"UnrealEd"
+			});
+		}
+	}
+}
diff --git a/Engine/Plugins/VirtualProduction/CaptureData/Source/CaptureDataUtils/Private/CaptureDataUtilsModule.cpp b/Engine/Plugins/VirtualProduction/CaptureData/Source/CaptureDataUtils/Private/CaptureDataUtilsModule.cpp
new file mode 100644
--- /dev/null
+++ b/Engine/Plugins/VirtualProduction/CaptureData/Source/CaptureDataUtils/Private/CaptureDataUtilsModule.cpp
@@ -0,0 +1,7 @@
+// Copyright Epic Games, Inc.All Rights Reserved.
+
+#include "Modules/ModuleManager.h"
+
+using FCaptureDataUtilsModule = IModuleInterface;
+
+IMPLEMENT_MODULE(FCaptureDataUtilsModule, CaptureDataUtils);
\ No newline at end of file
diff --git a/Engine/Plugins/VirtualProduction/CaptureData/Source/CaptureDataUtils/Private/ImageSequenceTimecodeUtils.cpp b/Engine/Plugins/VirtualProduction/CaptureData/Source/CaptureDataUtils/Private/ImageSequenceTimecodeUtils.cpp
new file mode 100644
--- /dev/null
+++ b/Engine/Plugins/VirtualProduction/CaptureData/Source/CaptureDataUtils/Private/ImageSequenceTimecodeUtils.cpp
@@ -0,0 +1,103 @@
+// Copyright Epic Games, Inc.All Rights Reserved.
+
+#include "ImageSequenceTimecodeUtils.h"
+
+#include "ParseTakeUtils.h"
+
+#if WITH_EDITOR
+#include "Subsystems/EditorAssetSubsystem.h"
+#include "Editor.h"
+#endif
+
+const FName UImageSequenceTimecodeUtils::TimecodeTagName = "Timecode";
+const FName UImageSequenceTimecodeUtils::TimecodeRateTagName = "TimecodeRate";
+
+void UImageSequenceTimecodeUtils::SetTimecodeInfo(const FTimecode& InTimecode, const FFrameRate& InFrameRate, UImgMediaSource* InImageSequence)
+{
+	FString ImageTimecodeString = InTimecode.ToString();
+	FString ImageTimecodeRateString = FString::SanitizeFloat(InFrameRate.AsDecimal());
+
+	SetTimecodeInfoString(ImageTimecodeString, ImageTimecodeRateString, InImageSequence);
+}
+
+void UImageSequenceTimecodeUtils::SetTimecodeInfoString(const FString& InTimecode, const FString& InFrameRate, UImgMediaSource* InImageSequence)
+{
+#if WITH_EDITOR
+	UEditorAssetSubsystem* EditorAssetSubsystem = GEditor ? GEditor->GetEditorSubsystem<UEditorAssetSubsystem>() : nullptr;
+
+	if (!EditorAssetSubsystem)
+	{
+		return;
+	}
+
+	EditorAssetSubsystem->SetMetadataTag(InImageSequence, UImageSequenceTimecodeUtils::TimecodeTagName, InTimecode);
+	EditorAssetSubsystem->SetMetadataTag(InImageSequence, UImageSequenceTimecodeUtils::TimecodeRateTagName, InFrameRate);
+#endif
+}
+
+TOptional<FTimecode> UImageSequenceTimecodeUtils::GetTimecode(UImgMediaSource* InImageSequence)
+{
+	TOptional<FString> TimecodeOpt = GetTimecodeString(InImageSequence);
+
+	if (!TimecodeOpt.IsSet())
+	{
+		return {};
+	}
+
+	return ParseTimecode(TimecodeOpt.GetValue());
+}
+
+TOptional<FFrameRate> UImageSequenceTimecodeUtils::GetFrameRate(UImgMediaSource* InImageSequence)
+{
+	TOptional<FString> TimecodeRateOpt = GetFrameRateString(InImageSequence);
+
+	if (!TimecodeRateOpt.IsSet())
+	{
+		return {};
+	}
+
+	double TimecodeRate = FCString::Atod(*TimecodeRateOpt.GetValue());
+
+	return ConvertFrameRate(TimecodeRate);
+}
+
+TOptional<FString> UImageSequenceTimecodeUtils::GetTimecodeString(UImgMediaSource* InImageSequence)
+{
+#if WITH_EDITOR
+	UEditorAssetSubsystem* EditorAssetSubsystem = GEditor ? GEditor->GetEditorSubsystem<UEditorAssetSubsystem>() : nullptr;
+
+	if (!EditorAssetSubsystem)
+	{
+		return {};
+	}
+
+	FString Timecode = EditorAssetSubsystem->GetMetadataTag(InImageSequence, UImageSequenceTimecodeUtils::TimecodeTagName);
+
+	return Timecode;
+#else
+	return {};
+#endif
+}
+
+TOptional<FString> UImageSequenceTimecodeUtils::GetFrameRateString(UImgMediaSource* InImageSequence)
+{
+#if WITH_EDITOR
+	UEditorAssetSubsystem* EditorAssetSubsystem = GEditor ? GEditor->GetEditorSubsystem<UEditorAssetSubsystem>() : nullptr;
+
+	if (!EditorAssetSubsystem)
+	{
+		return {};
+	}
+
+	FString TimecodeRateMetadata = EditorAssetSubsystem->GetMetadataTag(InImageSequence, UImageSequenceTimecodeUtils::TimecodeRateTagName);
+
+	if (!TimecodeRateMetadata.IsNumeric())
+	{
+		return {};
+	}
+	
+	return TimecodeRateMetadata;
+#else
+	return {};
+#endif
+}
\ No newline at end of file
diff --git a/Engine/Plugins/VirtualProduction/CaptureData/Source/CaptureDataUtils/Private/ParseTakeUtils.cpp b/Engine/Plugins/VirtualProduction/CaptureData/Source/CaptureDataUtils/Private/ParseTakeUtils.cpp
new file mode 100644
--- /dev/null
+++ b/Engine/Plugins/VirtualProduction/CaptureData/Source/CaptureDataUtils/Private/ParseTakeUtils.cpp
@@ -0,0 +1,51 @@
+// Copyright Epic Games, Inc.All Rights Reserved.
+
+#include "ParseTakeUtils.h"
+
+#include "CoreTypes.h"
+#include "Containers/Array.h"
+
+FTimecode ParseTimecode(const FString& InTimecodeString)
+{
+	TArray<FString> ParsedStringArray;
+	InTimecodeString.ParseIntoArray(ParsedStringArray, TEXT(":"));
+
+	if (ParsedStringArray.Num() != 4 && ParsedStringArray.Num() != 3)
+	{
+		return FTimecode(0, 0, 0, 0, false);
+	}
+
+	int32 Hours = FCString::Atoi(*ParsedStringArray[0]);
+	int32 Minutes = FCString::Atoi(*ParsedStringArray[1]);
+
+	int32 Seconds = 0;
+	int32 Frames = 0;
+	bool bIsFrameDrop = ParsedStringArray[2].Contains(TEXT(";"));
+	if (bIsFrameDrop)
+	{
+		// ParsedStringArray[2] == 00;00
+		TArray<FString> SecondsAndFrames;
+		ParsedStringArray[2].ParseIntoArray(SecondsAndFrames, TEXT(";"));
+
+		if (SecondsAndFrames.Num() != 2)
+		{
+			return FTimecode(0, 0, 0, 0, false);
+		}
+
+		Seconds = FCString::Atoi(*SecondsAndFrames[0]);
+		Frames = FCString::Atoi(*SecondsAndFrames[1]);
+	}
+	else
+	{
+		Seconds = FCString::Atoi(*ParsedStringArray[2]);
+		Frames = FCString::Atoi(*ParsedStringArray[3]);
+	}
+
+	return FTimecode(Hours, Minutes, Seconds, Frames, bIsFrameDrop);
+}
+
+FFrameRate ConvertFrameRate(double InFrameRate)
+{
+	return FFrameRate(FMath::CeilToInt(InFrameRate) * 1000,
+					  FMath::IsNearlyZero(FMath::Fractional(InFrameRate)) ? 1000 : 1001);
+}
\ No newline at end of file
diff --git a/Engine/Plugins/VirtualProduction/CaptureData/Source/CaptureDataUtils/Public/ImageSequenceTimecodeUtils.h b/Engine/Plugins/VirtualProduction/CaptureData/Source/CaptureDataUtils/Public/ImageSequenceTimecodeUtils.h
new file mode 100644
--- /dev/null
+++ b/Engine/Plugins/VirtualProduction/CaptureData/Source/CaptureDataUtils/Public/ImageSequenceTimecodeUtils.h
@@ -0,0 +1,38 @@
+// Copyright Epic Games, Inc. All Rights Reserved.
+
+#pragma once
+
+#include "Misc/Timecode.h"
+#include "ImgMediaSource.h"
+
+#include "ImageSequenceTimecodeUtils.generated.h"
+
+UCLASS(BlueprintType, Blueprintable)
+class CAPTUREDATAUTILS_API UImageSequenceTimecodeUtils
+	: public UObject
+{
+	GENERATED_BODY()
+
+public:
+
+	static const FName TimecodeTagName;
+	static const FName TimecodeRateTagName;
+
+	UFUNCTION(BlueprintCallable, Category = "ImageSequence")
+	static void SetTimecodeInfo(const FTimecode& InTimecode, const FFrameRate& InFrameRate, UImgMediaSource* InImageSequence);
+
+	UFUNCTION(BlueprintCallable, Category = "ImageSequence")
+	static void SetTimecodeInfoString(const FString& InTimecode, const FString& InFrameRate, UImgMediaSource* InImageSequence);
+
+	UFUNCTION(BlueprintCallable, Category = "ImageSequence")
+	static TOptional<FTimecode> GetTimecode(UImgMediaSource* InImageSequence);
+
+	UFUNCTION(BlueprintCallable, Category = "ImageSequence")
+	static TOptional<FFrameRate> GetFrameRate(UImgMediaSource* InImageSequence);
+
+	UFUNCTION(BlueprintCallable, Category = "ImageSequence")
+	static TOptional<FString> GetTimecodeString(UImgMediaSource* InImageSequence);
+
+	UFUNCTION(BlueprintCallable, Category = "ImageSequence")
+	static TOptional<FString> GetFrameRateString(UImgMediaSource* InImageSequence);
+};
diff --git a/Engine/Plugins/VirtualProduction/CaptureData/Source/CaptureDataUtils/Public/ParseTakeUtils.h b/Engine/Plugins/VirtualProduction/CaptureData/Source/CaptureDataUtils/Public/ParseTakeUtils.h
new file mode 100644
--- /dev/null
+++ b/Engine/Plugins/VirtualProduction/CaptureData/Source/CaptureDataUtils/Public/ParseTakeUtils.h
@@ -0,0 +1,11 @@
+// Copyright Epic Games, Inc. All Rights Reserved.
+
+#pragma once
+
+#include "Misc/Timecode.h"
+#include "Misc/FrameRate.h"
+
+#include "Containers/UnrealString.h"
+
+CAPTUREDATAUTILS_API FTimecode ParseTimecode(const FString& InTimecodeString);
+CAPTUREDATAUTILS_API FFrameRate ConvertFrameRate(double InFrameRate);
\ No newline at end of file
