Perform asset update on startup (with logs)

diff --git a/Engine/Plugins/VirtualProduction/CaptureData/Source/CaptureDataCore/Private/CaptureDataCoreModule.cpp b/Engine/Plugins/VirtualProduction/CaptureData/Source/CaptureDataCore/Private/CaptureDataCoreModule.cpp
--- a/Engine/Plugins/VirtualProduction/CaptureData/Source/CaptureDataCore/Private/CaptureDataCoreModule.cpp
+++ b/Engine/Plugins/VirtualProduction/CaptureData/Source/CaptureDataCore/Private/CaptureDataCoreModule.cpp
@@ -43,19 +43,16 @@ void FCaptureDataCoreModule::CheckAssetMigration()
 	// Prompt user to reload these assets
 	if (!Assets.IsEmpty())
 	{
-		const FText AssetMigrationDialogMessage = LOCTEXT("CaptureDataAssetMigrationDialog", "Found {0} Capture Data assets which need to be reloaded."
-			" Some MetaHuman Animator features may be unavailable until these assets are reloaded. \n\nReload now ? ");
-		const EAppReturnType::Type Answer = FMessageDialog::Open(EAppMsgType::YesNo, FText::Format(AssetMigrationDialogMessage, Assets.Num()));
-		if (Answer == EAppReturnType::No)
+		UE_LOG(LogCaptureDataCore, Display, TEXT("Found %d Capture Data assets which need to be updated. Starting update ..."), Assets.Num());
+		for (FAssetData AssetData : Assets)
 		{
-			return;
+			if (UObject* Asset = AssetData.GetAsset())
+			{
+				Asset->ReloadConfig();
+			}
 		}
-	}
-
-	for (FAssetData AssetData : Assets)
-	{
-		AssetData.GetAsset()->ReloadConfig();
-	}
+		UE_LOG(LogCaptureDataCore, Display, TEXT("Finished updating Capture Data assets"));
+	}	
 }
 
 IMPLEMENT_MODULE(FCaptureDataCoreModule, CaptureDataCore)
