Crash when clearing the audio field for the footage capture date and then adding this footage to a performance asset

diff --git a/Engine/Plugins/VirtualProduction/CaptureData/Source/CaptureDataCore/Private/CaptureData.cpp b/Engine/Plugins/VirtualProduction/CaptureData/Source/CaptureDataCore/Private/CaptureData.cpp
--- a/Engine/Plugins/VirtualProduction/CaptureData/Source/CaptureDataCore/Private/CaptureData.cpp
+++ b/Engine/Plugins/VirtualProduction/CaptureData/Source/CaptureDataCore/Private/CaptureData.cpp
@@ -365,19 +365,32 @@ void UFootageCaptureData::GetFrameRanges(const FFrameRate& InTargetRate, ETimeco
 
 	for (int32 Index = 0; Index < ImageSequences.Num(); ++Index)
 	{
-		OutMediaFrameRanges.Add(ImageSequences[Index], GetFrameRange(InTargetRate, ImageSequences[Index], GetEffectiveImageTimecode(Index), GetEffectiveImageTimecodeRate(Index), InTimecodeAlignment == ETimecodeAlignment::None));
+		const TObjectPtr<UImgMediaSource>& ImageSequence = ImageSequences[Index];
+
+		if (ImageSequence)
+		{
+			OutMediaFrameRanges.Add(ImageSequence, GetFrameRange(InTargetRate, ImageSequence, GetEffectiveImageTimecode(Index), GetEffectiveImageTimecodeRate(Index), InTimecodeAlignment == ETimecodeAlignment::None));
+		}
 	}
 
 	for (int32 Index = 0; Index < DepthSequences.Num(); ++Index)
 	{
-		OutMediaFrameRanges.Add(DepthSequences[Index], GetFrameRange(InTargetRate, DepthSequences[Index], GetEffectiveDepthTimecode(Index), GetEffectiveDepthTimecodeRate(Index), InTimecodeAlignment == ETimecodeAlignment::None));
+		const TObjectPtr<UImgMediaSource>& DepthSequence = DepthSequences[Index];
+
+		if (DepthSequence)
+		{
+			OutMediaFrameRanges.Add(DepthSequence, GetFrameRange(InTargetRate, DepthSequence, GetEffectiveDepthTimecode(Index), GetEffectiveDepthTimecodeRate(Index), InTimecodeAlignment == ETimecodeAlignment::None));
+		}
 	}
 
 	if (bInIncludeAudio)
 	{
 		for (TObjectPtr<class USoundWave> Audio : Audios)
 		{
-			OutMediaFrameRanges.Add(Audio, GetFrameRange(InTargetRate, Audio, GetEffectiveAudioTimecode(), GetEffectiveAudioTimecodeRate(), InTimecodeAlignment == ETimecodeAlignment::None));
+			if (Audio)
+			{
+				OutMediaFrameRanges.Add(Audio, GetFrameRange(InTargetRate, Audio, GetEffectiveAudioTimecode(), GetEffectiveAudioTimecodeRate(), InTimecodeAlignment == ETimecodeAlignment::None));
+			}
 		}
 	}
 
@@ -652,8 +665,6 @@ void UFootageCaptureData::PopulateCameraNames(UFootageCaptureData* InFootageCapt
 {
 	OutCameraNames.Reset();
 
-
-
 	if (InFootageCaptureData && !InFootageCaptureData->CameraCalibrations.IsEmpty())
 	{
 		const UCameraCalibration* CameraCalibration = InFootageCaptureData->CameraCalibrations[0];
diff --git a/Engine/Plugins/VirtualProduction/CaptureData/Source/CaptureDataUtils/Private/ImageSequenceTimecodeUtils.cpp b/Engine/Plugins/VirtualProduction/CaptureData/Source/CaptureDataUtils/Private/ImageSequenceTimecodeUtils.cpp
--- a/Engine/Plugins/VirtualProduction/CaptureData/Source/CaptureDataUtils/Private/ImageSequenceTimecodeUtils.cpp
+++ b/Engine/Plugins/VirtualProduction/CaptureData/Source/CaptureDataUtils/Private/ImageSequenceTimecodeUtils.cpp
@@ -37,6 +37,11 @@ void UImageSequenceTimecodeUtils::SetTimecodeInfoString(const FString& InTimecod
 
 TOptional<FTimecode> UImageSequenceTimecodeUtils::GetTimecode(UImgMediaSource* InImageSequence)
 {
+	if (!InImageSequence)
+	{
+		return {};
+	}
+
 	TOptional<FString> TimecodeOpt = GetTimecodeString(InImageSequence);
 
 	if (!TimecodeOpt.IsSet())
@@ -49,6 +54,11 @@ TOptional<FTimecode> UImageSequenceTimecodeUtils::GetTimecode(UImgMediaSource* I
 
 TOptional<FFrameRate> UImageSequenceTimecodeUtils::GetFrameRate(UImgMediaSource* InImageSequence)
 {
+	if (!InImageSequence)
+	{
+		return {};
+	}
+
 	TOptional<FString> TimecodeRateOpt = GetFrameRateString(InImageSequence);
 
 	if (!TimecodeRateOpt.IsSet())
@@ -73,6 +83,11 @@ TOptional<FString> UImageSequenceTimecodeUtils::GetTimecodeString(UImgMediaSourc
 
 	FString Timecode = EditorAssetSubsystem->GetMetadataTag(InImageSequence, UImageSequenceTimecodeUtils::TimecodeTagName);
 
+	if (Timecode.IsEmpty())
+	{
+		return {};
+	}
+
 	return Timecode;
 #else
 	return {};
