SetTimecode blueprint should default to current set value

#virtualized

diff --git a/Engine/Plugins/VirtualProduction/CaptureData/Source/CaptureDataCore/Private/CaptureData.cpp b/Engine/Plugins/VirtualProduction/CaptureData/Source/CaptureDataCore/Private/CaptureData.cpp
--- a/Engine/Plugins/VirtualProduction/CaptureData/Source/CaptureDataCore/Private/CaptureData.cpp
+++ b/Engine/Plugins/VirtualProduction/CaptureData/Source/CaptureDataCore/Private/CaptureData.cpp
@@ -473,7 +473,7 @@ FTimecode UFootageCaptureData::GetEffectiveImageTimecode(int32 InView) const
 	check(InView < DepthSequences.Num());
 
 #if WITH_EDITOR
-	TOptional<FTimecode> TimecodeOpt = UImageSequenceTimecodeUtils::GetTimecode(ImageSequences[InView]);
+	TOptional<FTimecode> TimecodeOpt = UImageSequenceTimecodeUtils::TryGetTimecode(ImageSequences[InView]);
 	if (TimecodeOpt.IsSet())
 	{
 		return MoveTemp(TimecodeOpt.GetValue());
@@ -493,7 +493,7 @@ FFrameRate UFootageCaptureData::GetEffectiveImageTimecodeRate(int32 InView) cons
 	check(InView < DepthSequences.Num());
 
 #if WITH_EDITOR
-	TOptional<FFrameRate> FrameRateOpt = UImageSequenceTimecodeUtils::GetFrameRate(ImageSequences[InView]);
+	TOptional<FFrameRate> FrameRateOpt = UImageSequenceTimecodeUtils::TryGetFrameRate(ImageSequences[InView]);
 	if (FrameRateOpt.IsSet())
 	{
 		return MoveTemp(FrameRateOpt.GetValue());
@@ -513,7 +513,7 @@ FTimecode UFootageCaptureData::GetEffectiveDepthTimecode(int32 InView) const
 	check(InView < DepthSequences.Num());
 
 #if WITH_EDITOR
-	TOptional<FTimecode> TimecodeOpt = UImageSequenceTimecodeUtils::GetTimecode(DepthSequences[InView]);
+	TOptional<FTimecode> TimecodeOpt = UImageSequenceTimecodeUtils::TryGetTimecode(DepthSequences[InView]);
 	if (TimecodeOpt.IsSet())
 	{
 		return MoveTemp(TimecodeOpt.GetValue());
@@ -533,7 +533,7 @@ FFrameRate UFootageCaptureData::GetEffectiveDepthTimecodeRate(int32 InView) cons
 	check(InView < DepthSequences.Num());
 
 #if WITH_EDITOR
-	TOptional<FFrameRate> FrameRateOpt = UImageSequenceTimecodeUtils::GetFrameRate(DepthSequences[InView]);
+	TOptional<FFrameRate> FrameRateOpt = UImageSequenceTimecodeUtils::TryGetFrameRate(DepthSequences[InView]);
 	if (FrameRateOpt.IsSet())
 	{
 		return MoveTemp(FrameRateOpt.GetValue());
@@ -585,8 +585,8 @@ void UFootageCaptureData::GetDefaultTimecodeInfo(FTimecode& OutTimecode, FFrameR
 #if WITH_EDITOR
 	for (const TObjectPtr<class UImgMediaSource>& ImageSequence : ImageSequences)
 	{
-		TOptional<FTimecode> TimecodeOpt = UImageSequenceTimecodeUtils::GetTimecode(ImageSequence);
-		TOptional<FFrameRate> FrameRateOpt = UImageSequenceTimecodeUtils::GetFrameRate(ImageSequence);
+		TOptional<FTimecode> TimecodeOpt = UImageSequenceTimecodeUtils::TryGetTimecode(ImageSequence);
+		TOptional<FFrameRate> FrameRateOpt = UImageSequenceTimecodeUtils::TryGetFrameRate(ImageSequence);
 
 		if (TimecodeOpt.IsSet() && FrameRateOpt.IsSet())
 		{
@@ -598,8 +598,8 @@ void UFootageCaptureData::GetDefaultTimecodeInfo(FTimecode& OutTimecode, FFrameR
 
 	for (const TObjectPtr<class UImgMediaSource>& DepthSequence : DepthSequences)
 	{
-		TOptional<FTimecode> TimecodeOpt = UImageSequenceTimecodeUtils::GetTimecode(DepthSequence);
-		TOptional<FFrameRate> FrameRateOpt = UImageSequenceTimecodeUtils::GetFrameRate(DepthSequence);
+		TOptional<FTimecode> TimecodeOpt = UImageSequenceTimecodeUtils::TryGetTimecode(DepthSequence);
+		TOptional<FFrameRate> FrameRateOpt = UImageSequenceTimecodeUtils::TryGetFrameRate(DepthSequence);
 
 		if (TimecodeOpt.IsSet() && FrameRateOpt.IsSet())
 		{
diff --git a/Engine/Plugins/VirtualProduction/CaptureData/Source/CaptureDataUtils/Private/ImageSequenceTimecodeUtils.cpp b/Engine/Plugins/VirtualProduction/CaptureData/Source/CaptureDataUtils/Private/ImageSequenceTimecodeUtils.cpp
--- a/Engine/Plugins/VirtualProduction/CaptureData/Source/CaptureDataUtils/Private/ImageSequenceTimecodeUtils.cpp
+++ b/Engine/Plugins/VirtualProduction/CaptureData/Source/CaptureDataUtils/Private/ImageSequenceTimecodeUtils.cpp
@@ -35,31 +35,69 @@ void UImageSequenceTimecodeUtils::SetTimecodeInfoString(const FString& InTimecod
 #endif
 }
 
-TOptional<FTimecode> UImageSequenceTimecodeUtils::GetTimecode(UImgMediaSource* InImageSequence)
+FTimecode UImageSequenceTimecodeUtils::GetTimecode(UImgMediaSource* InImageSequence)
 {
-	if (!InImageSequence)
+	TOptional<FTimecode> TimecodeOpt = TryGetTimecode(InImageSequence);
+
+	if (TimecodeOpt.IsSet())
 	{
-		return {};
+		return TimecodeOpt.GetValue();
 	}
 
-	TOptional<FString> TimecodeOpt = GetTimecodeString(InImageSequence);
+	return FTimecode();
+}
 
-	if (!TimecodeOpt.IsSet())
+FFrameRate UImageSequenceTimecodeUtils::GetFrameRate(UImgMediaSource* InImageSequence)
+{
+	TOptional<FFrameRate> FrameRateOpt = TryGetFrameRate(InImageSequence);
+
+	if (FrameRateOpt.IsSet())
 	{
-		return {};
+		return FrameRateOpt.GetValue();
 	}
 
-	return ParseTimecode(TimecodeOpt.GetValue());
+	return FFrameRate();
 }
 
-TOptional<FFrameRate> UImageSequenceTimecodeUtils::GetFrameRate(UImgMediaSource* InImageSequence)
+FString UImageSequenceTimecodeUtils::GetTimecodeString(UImgMediaSource* InImageSequence)
 {
-	if (!InImageSequence)
+	TOptional<FString> TimecodeOpt = TryGetTimecodeString(InImageSequence);
+
+	if (TimecodeOpt.IsSet())
+	{
+		return TimecodeOpt.GetValue();
+	}
+
+	return FString();
+}
+
+FString UImageSequenceTimecodeUtils::GetFrameRateString(UImgMediaSource* InImageSequence)
+{
+	TOptional<FString> FrameRateOpt = TryGetFrameRateString(InImageSequence);
+
+	if (FrameRateOpt.IsSet())
+	{
+		return FrameRateOpt.GetValue();
+	}
+
+	return FString();
+}
+
+TOptional<FTimecode> UImageSequenceTimecodeUtils::TryGetTimecode(UImgMediaSource* InImageSequence)
+{
+	TOptional<FString> TimecodeOpt = TryGetTimecodeString(InImageSequence);
+
+	if (!TimecodeOpt.IsSet())
 	{
 		return {};
 	}
 
-	TOptional<FString> TimecodeRateOpt = GetFrameRateString(InImageSequence);
+	return ParseTimecode(TimecodeOpt.GetValue());
+}
+
+TOptional<FFrameRate> UImageSequenceTimecodeUtils::TryGetFrameRate(UImgMediaSource* InImageSequence)
+{
+	TOptional<FString> TimecodeRateOpt = TryGetFrameRateString(InImageSequence);
 
 	if (!TimecodeRateOpt.IsSet())
 	{
@@ -71,9 +109,14 @@ TOptional<FFrameRate> UImageSequenceTimecodeUtils::GetFrameRate(UImgMediaSource*
 	return ConvertFrameRate(TimecodeRate);
 }
 
-TOptional<FString> UImageSequenceTimecodeUtils::GetTimecodeString(UImgMediaSource* InImageSequence)
+TOptional<FString> UImageSequenceTimecodeUtils::TryGetTimecodeString(UImgMediaSource* InImageSequence)
 {
 #if WITH_EDITOR
+	if (!InImageSequence)
+	{
+		return {};
+	}
+
 	UEditorAssetSubsystem* EditorAssetSubsystem = GEditor ? GEditor->GetEditorSubsystem<UEditorAssetSubsystem>() : nullptr;
 
 	if (!EditorAssetSubsystem)
@@ -94,9 +137,14 @@ TOptional<FString> UImageSequenceTimecodeUtils::GetTimecodeString(UImgMediaSourc
 #endif
 }
 
-TOptional<FString> UImageSequenceTimecodeUtils::GetFrameRateString(UImgMediaSource* InImageSequence)
+TOptional<FString> UImageSequenceTimecodeUtils::TryGetFrameRateString(UImgMediaSource* InImageSequence)
 {
 #if WITH_EDITOR
+	if (!InImageSequence)
+	{
+		return {};
+	}
+
 	UEditorAssetSubsystem* EditorAssetSubsystem = GEditor ? GEditor->GetEditorSubsystem<UEditorAssetSubsystem>() : nullptr;
 
 	if (!EditorAssetSubsystem)
@@ -115,4 +163,4 @@ TOptional<FString> UImageSequenceTimecodeUtils::GetFrameRateString(UImgMediaSour
 #else
 	return {};
 #endif
-}
\ No newline at end of file
+}
diff --git a/Engine/Plugins/VirtualProduction/CaptureData/Source/CaptureDataUtils/Public/ImageSequenceTimecodeUtils.h b/Engine/Plugins/VirtualProduction/CaptureData/Source/CaptureDataUtils/Public/ImageSequenceTimecodeUtils.h
--- a/Engine/Plugins/VirtualProduction/CaptureData/Source/CaptureDataUtils/Public/ImageSequenceTimecodeUtils.h
+++ b/Engine/Plugins/VirtualProduction/CaptureData/Source/CaptureDataUtils/Public/ImageSequenceTimecodeUtils.h
@@ -7,6 +7,21 @@
 
 #include "ImageSequenceTimecodeUtils.generated.h"
 
+UCLASS(BlueprintType, Blueprintable)
+class UImageSequenceTimecodeInfo
+	: public UObject
+{
+	GENERATED_BODY()
+
+public:
+
+	UPROPERTY(EditAnywhere, BlueprintReadWrite, meta=(ExposeOnSpawn), Category = "ImageSequenceInfo")
+	FTimecode Timecode;
+
+	UPROPERTY(EditAnywhere, BlueprintReadWrite, meta=(ExposeOnSpawn), Category = "ImageSequenceInfo")
+	FFrameRate FrameRate;
+};
+
 UCLASS(BlueprintType, Blueprintable)
 class CAPTUREDATAUTILS_API UImageSequenceTimecodeUtils
 	: public UObject
@@ -25,14 +40,19 @@ public:
 	static void SetTimecodeInfoString(const FString& InTimecode, const FString& InFrameRate, UImgMediaSource* InImageSequence);
 
 	UFUNCTION(BlueprintCallable, Category = "ImageSequence")
-	static TOptional<FTimecode> GetTimecode(UImgMediaSource* InImageSequence);
+	static FTimecode GetTimecode(UImgMediaSource* InImageSequence);
 
 	UFUNCTION(BlueprintCallable, Category = "ImageSequence")
-	static TOptional<FFrameRate> GetFrameRate(UImgMediaSource* InImageSequence);
+	static FFrameRate GetFrameRate(UImgMediaSource* InImageSequence);
 
 	UFUNCTION(BlueprintCallable, Category = "ImageSequence")
-	static TOptional<FString> GetTimecodeString(UImgMediaSource* InImageSequence);
+	static FString GetTimecodeString(UImgMediaSource* InImageSequence);
 
 	UFUNCTION(BlueprintCallable, Category = "ImageSequence")
-	static TOptional<FString> GetFrameRateString(UImgMediaSource* InImageSequence);
+	static FString GetFrameRateString(UImgMediaSource* InImageSequence);
+
+	static TOptional<FTimecode> TryGetTimecode(UImgMediaSource* InImageSequence);
+	static TOptional<FFrameRate> TryGetFrameRate(UImgMediaSource* InImageSequence);
+	static TOptional<FString> TryGetTimecodeString(UImgMediaSource* InImageSequence);
+	static TOptional<FString> TryGetFrameRateString(UImgMediaSource* InImageSequence);
 };
