Prevent depth selection for metahuman identity and performance processing

diff --git a/Engine/Plugins/VirtualProduction/CaptureData/Source/CaptureDataCore/Private/CaptureData.cpp b/Engine/Plugins/VirtualProduction/CaptureData/Source/CaptureDataCore/Private/CaptureData.cpp
--- a/Engine/Plugins/VirtualProduction/CaptureData/Source/CaptureDataCore/Private/CaptureData.cpp
+++ b/Engine/Plugins/VirtualProduction/CaptureData/Source/CaptureDataCore/Private/CaptureData.cpp
@@ -672,7 +672,10 @@ void UFootageCaptureData::PopulateCameraNames(UFootageCaptureData* InFootageCapt
 		{
 			for (const FExtendedLensFile& LensFile : CameraCalibration->CameraCalibrations)
 			{
-				OutCameraNames.Add(MakeShared<FString>(LensFile.Name));
+				if (!LensFile.IsDepthCamera)
+				{
+					OutCameraNames.Add(MakeShared<FString>(LensFile.Name));
+				}
 			}
 		}
 	}
