Crash occurs when CaptureData links to other assets are broken

diff --git a/Engine/Plugins/VirtualProduction/CaptureData/Source/CaptureDataCore/Private/CameraCalibration.cpp b/Engine/Plugins/VirtualProduction/CaptureData/Source/CaptureDataCore/Private/CameraCalibration.cpp
--- a/Engine/Plugins/VirtualProduction/CaptureData/Source/CaptureDataCore/Private/CameraCalibration.cpp
+++ b/Engine/Plugins/VirtualProduction/CaptureData/Source/CaptureDataCore/Private/CameraCalibration.cpp
@@ -36,13 +36,20 @@ void UCameraCalibration::PostLoad()
 	{
 		if (CameraCalibrations[0].Name.IsEmpty() && CameraCalibrations[1].Name.IsEmpty())
 		{
-			if (CameraCalibrations[0].LensFile->LensInfo.ImageDimensions.X == CameraCalibrations[1].LensFile->LensInfo.ImageDimensions.X * 2)
+			if (CameraCalibrations[0].LensFile)
 			{
-				CameraCalibrations[0].Name = TEXT("iPhone");
+				if (CameraCalibrations[0].LensFile->LensInfo.ImageDimensions.X == CameraCalibrations[1].LensFile->LensInfo.ImageDimensions.X * 2)
+				{
+					CameraCalibrations[0].Name = TEXT("iPhone");
+				}
+				else
+				{
+					CameraCalibrations[0].Name = TEXT("bot");
+				}
 			}
 			else
 			{
-				CameraCalibrations[0].Name = TEXT("bot");
+				CameraCalibrations[0].Name = TEXT("Unknown");
 			}
 
 			CameraCalibrations[1].Name = TEXT("Depth");
