MetaHuman: Use audio timecode and frame rate from soundwave asset.

Also prepare capture data for adding metadata to the assets and metahuman performance for reading it.

#virtualized

diff --git a/Engine/Plugins/VirtualProduction/CaptureData/Source/CaptureDataCore/Private/CaptureData.cpp b/Engine/Plugins/VirtualProduction/CaptureData/Source/CaptureDataCore/Private/CaptureData.cpp
--- a/Engine/Plugins/VirtualProduction/CaptureData/Source/CaptureDataCore/Private/CaptureData.cpp
+++ b/Engine/Plugins/VirtualProduction/CaptureData/Source/CaptureDataCore/Private/CaptureData.cpp
@@ -7,6 +7,8 @@
 
 #include "ImageSequenceUtils.h"
 
+#include "SoundWaveTimecodeUtils.h"
+
 #include "Components/SkeletalMeshComponent.h"
 #include "Components/StaticMeshComponent.h"
 #include "Rendering/SkeletalMeshModel.h"
@@ -363,13 +365,19 @@ void UFootageCaptureData::GetFrameRanges(const FFrameRate& InTargetRate, ETimeco
 	OutProcessingFrameRange = TRange<FFrameNumber>(0, 0);
 	OutMaximumFrameRange = TRange<FFrameNumber>(0, 0);
 
+	const bool bMediaStartFrameIsZero = InTimecodeAlignment == ETimecodeAlignment::None;
+
 	for (int32 Index = 0; Index < ImageSequences.Num(); ++Index)
 	{
 		const TObjectPtr<UImgMediaSource>& ImageSequence = ImageSequences[Index];
 
 		if (ImageSequence)
 		{
-			OutMediaFrameRanges.Add(ImageSequence, GetFrameRange(InTargetRate, ImageSequence, GetEffectiveImageTimecode(Index), GetEffectiveImageTimecodeRate(Index), InTimecodeAlignment == ETimecodeAlignment::None));
+			const FTimecode EffectiveImageTimecode = GetEffectiveImageTimecode(Index);
+			const FFrameRate EffectiveImageTimecodeRate = GetEffectiveImageTimecodeRate(Index);
+			const TRange<FFrameNumber> ImageFrameFrange = GetFrameRange(InTargetRate, ImageSequence, EffectiveImageTimecode, EffectiveImageTimecodeRate, bMediaStartFrameIsZero);
+
+			OutMediaFrameRanges.Add(ImageSequence, ImageFrameFrange);
 		}
 	}
 
@@ -379,7 +387,11 @@ void UFootageCaptureData::GetFrameRanges(const FFrameRate& InTargetRate, ETimeco
 
 		if (DepthSequence)
 		{
-			OutMediaFrameRanges.Add(DepthSequence, GetFrameRange(InTargetRate, DepthSequence, GetEffectiveDepthTimecode(Index), GetEffectiveDepthTimecodeRate(Index), InTimecodeAlignment == ETimecodeAlignment::None));
+			const FTimecode EffectiveDepthTimecode = GetEffectiveDepthTimecode(Index);
+			const FFrameRate EffectiveDepthTimecodeRate = GetEffectiveDepthTimecodeRate(Index);
+			const TRange<FFrameNumber> DepthFrameRange = GetFrameRange(InTargetRate, DepthSequence, EffectiveDepthTimecode, EffectiveDepthTimecodeRate, bMediaStartFrameIsZero);
+
+			OutMediaFrameRanges.Add(DepthSequence, DepthFrameRange);
 		}
 	}
 
@@ -389,7 +401,11 @@ void UFootageCaptureData::GetFrameRanges(const FFrameRate& InTargetRate, ETimeco
 		{
 			if (Audio)
 			{
-				OutMediaFrameRanges.Add(Audio, GetFrameRange(InTargetRate, Audio, GetEffectiveAudioTimecode(), GetEffectiveAudioTimecodeRate(), InTimecodeAlignment == ETimecodeAlignment::None));
+				const FTimecode EffectiveAudioTimecode = GetEffectiveAudioTimecode();
+				const FFrameRate EffectiveAudioTimecodeRate = GetEffectiveAudioTimecodeRate();
+				const TRange<FFrameNumber> AudioFrameRange = GetFrameRange(InTargetRate, Audio, EffectiveAudioTimecode, EffectiveAudioTimecodeRate, bMediaStartFrameIsZero);
+
+				OutMediaFrameRanges.Add(Audio, AudioFrameRange);
 			}
 		}
 	}
@@ -549,12 +565,22 @@ FFrameRate UFootageCaptureData::GetEffectiveDepthTimecodeRate(int32 InView) cons
 FTimecode UFootageCaptureData::GetEffectiveAudioTimecode() const
 {
 	FTimecode Timecode;
+	bool TimecodeFromAsset = false;
 
-	if (bAudioTimecodePresent)
+	if (Audios.Num() > 0)
 	{
-		Timecode = AudioTimecode;
+		const USoundWave* SoundWave = Audios[0].Get();
+		check(IsValid(SoundWave));
+
+		TOptional<FTimecode> TimecodeOpt = USoundWaveTimecodeUtils::GetTimecode(SoundWave);
+		if (TimecodeOpt.IsSet())
+		{
+			Timecode = TimecodeOpt.GetValue();
+			TimecodeFromAsset = true;
+		}
 	}
-	else
+
+	if (!TimecodeFromAsset)
 	{
 		FFrameRate FrameRate;
 		GetDefaultTimecodeInfo(Timecode, FrameRate);
@@ -566,12 +592,22 @@ FTimecode UFootageCaptureData::GetEffectiveAudioTimecode() const
 FFrameRate UFootageCaptureData::GetEffectiveAudioTimecodeRate() const
 {
 	FFrameRate FrameRate;
+	bool TimecodeRateFromAsset = false;
 
-	if (bAudioTimecodePresent)
+	if (Audios.Num() > 0)
 	{
-		FrameRate = AudioTimecodeRate;
+		const USoundWave* SoundWave = Audios[0].Get();
+		check(IsValid(SoundWave));
+
+		TOptional<FFrameRate> FrameRateOpt = USoundWaveTimecodeUtils::GetFrameRate(SoundWave);
+		if (FrameRateOpt.IsSet())
+		{
+			FrameRate = FrameRateOpt.GetValue();
+			TimecodeRateFromAsset = true;
+		}
 	}
-	else
+
+	if (!TimecodeRateFromAsset)
 	{
 		FTimecode Timecode;
 		GetDefaultTimecodeInfo(Timecode, FrameRate);
@@ -610,11 +646,21 @@ void UFootageCaptureData::GetDefaultTimecodeInfo(FTimecode& OutTimecode, FFrameR
 	}
 #endif
 
-	if (bAudioTimecodePresent)
+
+	if (Audios.Num() > 0)
 	{
-		OutTimecode = AudioTimecode;
-		OutFrameRate = AudioTimecodeRate;
-		return;
+		const USoundWave* SoundWave = Audios[0].Get();
+		check(IsValid(SoundWave));
+
+		TOptional<FTimecode> TimecodeOpt = USoundWaveTimecodeUtils::GetTimecode(SoundWave);
+		TOptional<FFrameRate> FrameRateOpt = USoundWaveTimecodeUtils::GetFrameRate(SoundWave);
+
+		if (TimecodeOpt.IsSet() && FrameRateOpt.IsSet())
+		{
+			OutTimecode = TimecodeOpt.GetValue();
+			OutFrameRate = FrameRateOpt.GetValue();
+			return;
+		}
 	}
 
 	OutTimecode = FTimecode(0, 0, 0, 0, false);
@@ -659,6 +705,14 @@ void UFootageCaptureData::PostLoad()
 		DepthSequences.Add(MoveTemp(View.DepthSequence));
 	}
 #endif
+	if (Audios.Num() > 0 && bAudioTimecodePresent_DEPRECATED == true) 
+	{
+		USoundWave* SoundWave = Audios[0].Get();
+		check(IsValid(SoundWave));
+
+		USoundWaveTimecodeUtils::SetTimecodeInfo(AudioTimecode_DEPRECATED, AudioTimecodeRate_DEPRECATED, SoundWave);
+		bAudioTimecodePresent_DEPRECATED = false;
+	}
 }
 
 void UFootageCaptureData::PopulateCameraNames(UFootageCaptureData* InFootageCaptureData, FString& InOutCamera, TArray<TSharedPtr<FString>>& OutCameraNames)
diff --git a/Engine/Plugins/VirtualProduction/CaptureData/Source/CaptureDataCore/Public/CaptureData.h b/Engine/Plugins/VirtualProduction/CaptureData/Source/CaptureDataCore/Public/CaptureData.h
--- a/Engine/Plugins/VirtualProduction/CaptureData/Source/CaptureDataCore/Public/CaptureData.h
+++ b/Engine/Plugins/VirtualProduction/CaptureData/Source/CaptureDataCore/Public/CaptureData.h
@@ -9,7 +9,6 @@
 #include "Misc/Paths.h"
 #include "Engine/StaticMesh.h"
 #include "FrameRange.h"
-
 #include "Templates/ValueOrError.h"
 
 #include "CaptureData.generated.h"
@@ -221,14 +220,14 @@ public:
 	UPROPERTY(meta = (DeprecatedProperty, DeprecationMessage = "Audio is deprecated. Please use Audios instead."))
 	TObjectPtr<class USoundWave> Audio_DEPRECATED;
 
-	UPROPERTY(EditAnywhere, BlueprintReadWrite, Category = "Capture")
-	bool bAudioTimecodePresent = false;
+	UPROPERTY(meta = (DeprecatedProperty, DeprecationMessage = "AudioTimecodePresent is deprecated."))
+	bool bAudioTimecodePresent_DEPRECATED = false;
 
-	UPROPERTY(EditAnywhere, BlueprintReadWrite, Category = "Capture")
-	FTimecode AudioTimecode = FTimecode(0, 0, 0, 0, false);
+	UPROPERTY(meta = (DeprecatedProperty, DeprecationMessage = "AudioTimecode is deprecated."))
+	FTimecode AudioTimecode_DEPRECATED = FTimecode(0, 0, 0, 0, false);
 
-	UPROPERTY(EditAnywhere, BlueprintReadWrite, Category = "Capture")
-	FFrameRate AudioTimecodeRate = FFrameRate(30, 1);
+	UPROPERTY(meta = (DeprecatedProperty, DeprecationMessage = "AudioTimecodeRate is deprecated."))
+	FFrameRate AudioTimecodeRate_DEPRECATED = FFrameRate(30, 1);
 
 	UPROPERTY(EditAnywhere, BlueprintReadWrite, Category = "Capture")
 	TArray<TObjectPtr<class UCameraCalibration>> CameraCalibrations;
@@ -280,7 +279,6 @@ public:
 	static void PopulateCameraNames(UFootageCaptureData* InFootageCaptureData, FString& InOutCamera, TArray<TSharedPtr<FString>>& OutCameraNames);
 
 private:
-
 	/** Checks if camera calibration and views are valid. */
 	FVerifyResult ViewsContainsValidData() const;
 	FVerifyResult MetadataContainsValidData() const;
diff --git a/Engine/Plugins/VirtualProduction/CaptureData/Source/CaptureDataEditor/Private/Customizations/CaptureDataCustomizations.cpp b/Engine/Plugins/VirtualProduction/CaptureData/Source/CaptureDataEditor/Private/Customizations/CaptureDataCustomizations.cpp
--- a/Engine/Plugins/VirtualProduction/CaptureData/Source/CaptureDataEditor/Private/Customizations/CaptureDataCustomizations.cpp
+++ b/Engine/Plugins/VirtualProduction/CaptureData/Source/CaptureDataEditor/Private/Customizations/CaptureDataCustomizations.cpp
@@ -16,9 +16,6 @@ void FFootageCaptureDataCustomization::CustomizeDetails(IDetailLayoutBuilder& In
 	InDetailBuilder.AddPropertyToCategory(InDetailBuilder.GetProperty(GET_MEMBER_NAME_CHECKED(UFootageCaptureData, ImageSequences)));
 	InDetailBuilder.AddPropertyToCategory(InDetailBuilder.GetProperty(GET_MEMBER_NAME_CHECKED(UFootageCaptureData, DepthSequences)));
 	InDetailBuilder.AddPropertyToCategory(InDetailBuilder.GetProperty(GET_MEMBER_NAME_CHECKED(UFootageCaptureData, Audios)));
-	InDetailBuilder.AddPropertyToCategory(InDetailBuilder.GetProperty(GET_MEMBER_NAME_CHECKED(UFootageCaptureData, bAudioTimecodePresent)));
-	InDetailBuilder.AddPropertyToCategory(InDetailBuilder.GetProperty(GET_MEMBER_NAME_CHECKED(UFootageCaptureData, AudioTimecode)));
-	InDetailBuilder.AddPropertyToCategory(InDetailBuilder.GetProperty(GET_MEMBER_NAME_CHECKED(UFootageCaptureData, AudioTimecodeRate)));
 	InDetailBuilder.AddPropertyToCategory(InDetailBuilder.GetProperty(GET_MEMBER_NAME_CHECKED(UFootageCaptureData, CameraCalibrations)));
 	InDetailBuilder.AddPropertyToCategory(InDetailBuilder.GetProperty(GET_MEMBER_NAME_CHECKED(UFootageCaptureData, Metadata)));
 
diff --git a/Engine/Plugins/VirtualProduction/CaptureData/Source/CaptureDataUtils/Private/CaptureTimecodeInfo.h b/Engine/Plugins/VirtualProduction/CaptureData/Source/CaptureDataUtils/Private/CaptureTimecodeInfo.h
new file mode 100644
--- /dev/null
+++ b/Engine/Plugins/VirtualProduction/CaptureData/Source/CaptureDataUtils/Private/CaptureTimecodeInfo.h
@@ -0,0 +1,22 @@
+// Copyright Epic Games, Inc. All Rights Reserved.
+
+#pragma once
+
+#include "UObject/Object.h"
+#include "Misc/Timecode.h"
+#include "Misc/FrameRate.h"
+
+#include "CaptureTimecodeInfo.generated.h"
+
+UCLASS(BlueprintType, Blueprintable)
+class UCaptureTimecodeInfo
+	: public UObject
+{
+	GENERATED_BODY()
+public:
+	UPROPERTY(EditAnywhere, BlueprintReadWrite, meta = (ExposeOnSpawn), Category = "Info")
+	FTimecode Timecode;
+
+	UPROPERTY(EditAnywhere, BlueprintReadWrite, meta = (ExposeOnSpawn), Category = "Info")
+	FFrameRate FrameRate;
+};
diff --git a/Engine/Plugins/VirtualProduction/CaptureData/Source/CaptureDataUtils/Private/SoundWaveTimecodeUtils.cpp b/Engine/Plugins/VirtualProduction/CaptureData/Source/CaptureDataUtils/Private/SoundWaveTimecodeUtils.cpp
new file mode 100644
--- /dev/null
+++ b/Engine/Plugins/VirtualProduction/CaptureData/Source/CaptureDataUtils/Private/SoundWaveTimecodeUtils.cpp
@@ -0,0 +1,97 @@
+// Copyright Epic Games, Inc.All Rights Reserved.
+
+#include "SoundWaveTimecodeUtils.h"
+
+#include "Sound/SoundWave.h"
+#include "Sound/SoundWaveTimecodeInfo.h"
+
+DEFINE_LOG_CATEGORY_STATIC(LogSoundWaveTimecodeUtils, Log, All);
+
+void USoundWaveTimecodeUtils::SetTimecodeInfo(const FTimecode& InTimecode, const FFrameRate& InFrameRate, USoundWave* OutSoundWave)
+{
+	if (!ensure(IsValid(OutSoundWave)))
+	{
+		UE_LOG(LogSoundWaveTimecodeUtils, Error, TEXT("[%s] Sound wave pointer is invalid"), ANSI_TO_TCHAR(__FUNCTION__));
+		return;
+	}
+
+	// Sanity check that the frame rate has actually been set and is not just default constructed (60'000 fps)
+	check(InFrameRate != FFrameRate());
+
+	const float SampleRate = OutSoundWave->GetSampleRateForCurrentPlatform();
+	check(SampleRate > 0.0f);
+
+#if WITH_EDITORONLY_DATA
+	FSoundWaveTimecodeInfo SoundWaveTimecodeInfo;
+	SoundWaveTimecodeInfo.NumSamplesPerSecond = SampleRate;
+	SoundWaveTimecodeInfo.TimecodeRate = InFrameRate;
+	SoundWaveTimecodeInfo.NumSamplesSinceMidnight = InTimecode.ToTimespan(InFrameRate).GetTotalSeconds() * SoundWaveTimecodeInfo.NumSamplesPerSecond;
+	SoundWaveTimecodeInfo.bTimecodeIsDropFrame = InTimecode.bDropFrameFormat;
+	OutSoundWave->SetTimecodeInfo(SoundWaveTimecodeInfo);
+#else 
+	// This function may not really belong in this module...
+	check(false);
+	UE_LOG(LogSoundWaveTimecodeUtils, Error, TEXT("Function %s is not supported in non-editor targets"), ANSI_TO_TCHAR(__FUNCTION__));
+#endif
+}
+
+FTimecode USoundWaveTimecodeUtils::GetTimecode(const USoundWave* InSoundWave)
+{
+	if (!ensure(IsValid(InSoundWave)))
+	{
+		UE_LOG(LogSoundWaveTimecodeUtils, Error, TEXT("[%s] Sound wave pointer is invalid"), ANSI_TO_TCHAR(__FUNCTION__));
+		return {};
+	}
+
+	TOptional<FSoundWaveTimecodeInfo> SoundWaveTimecodeInfo;
+
+#if WITH_EDITORONLY_DATA
+	SoundWaveTimecodeInfo = InSoundWave->GetTimecodeInfo();
+#else 
+	// This function may not really belong in this module...
+	check(false);
+	UE_LOG(LogSoundWaveTimecodeUtils, Error, TEXT("Function %s is not supported in non-editor targets"), ANSI_TO_TCHAR(__FUNCTION__));
+#endif
+
+	if (!SoundWaveTimecodeInfo)
+	{
+		return {};
+	}
+
+	// Sanity check that the frame rate has actually been set and is not just default constructed (60'000 fps)
+	check(SoundWaveTimecodeInfo->TimecodeRate != FFrameRate());
+
+	// GetNumSecondsSinceMidnight will return 0.0 in this case, so nothing here will fail, but it may be useful to fail early during dev
+	check(SoundWaveTimecodeInfo->NumSamplesPerSecond > 0);
+
+	FTimecode Timecode(SoundWaveTimecodeInfo->GetNumSecondsSinceMidnight(), SoundWaveTimecodeInfo->TimecodeRate, true);
+	return Timecode;
+}
+
+FFrameRate USoundWaveTimecodeUtils::GetFrameRate(const USoundWave* InSoundWave)
+{
+	if (!ensure(IsValid(InSoundWave)))
+	{
+		UE_LOG(LogSoundWaveTimecodeUtils, Error, TEXT("[%s] Sound wave pointer is invalid"), ANSI_TO_TCHAR(__FUNCTION__));
+		return {};
+	}
+
+	TOptional<FSoundWaveTimecodeInfo> SoundWaveTimecodeInfo;
+
+#if WITH_EDITORONLY_DATA
+	SoundWaveTimecodeInfo = InSoundWave->GetTimecodeInfo();
+#else 
+	// This function may not really belong in this module...
+	check(false);
+	UE_LOG(LogSoundWaveTimecodeUtils, Error, TEXT("Function %s is not supported in non-editor targets"), ANSI_TO_TCHAR(__FUNCTION__));
+#endif
+
+	if (!SoundWaveTimecodeInfo)
+	{
+		return {};
+	}
+
+	return SoundWaveTimecodeInfo->TimecodeRate;
+}
+
+
diff --git a/Engine/Plugins/VirtualProduction/CaptureData/Source/CaptureDataUtils/Public/ImageSequenceTimecodeUtils.h b/Engine/Plugins/VirtualProduction/CaptureData/Source/CaptureDataUtils/Public/ImageSequenceTimecodeUtils.h
--- a/Engine/Plugins/VirtualProduction/CaptureData/Source/CaptureDataUtils/Public/ImageSequenceTimecodeUtils.h
+++ b/Engine/Plugins/VirtualProduction/CaptureData/Source/CaptureDataUtils/Public/ImageSequenceTimecodeUtils.h
@@ -7,21 +7,6 @@
 
 #include "ImageSequenceTimecodeUtils.generated.h"
 
-UCLASS(BlueprintType, Blueprintable)
-class UImageSequenceTimecodeInfo
-	: public UObject
-{
-	GENERATED_BODY()
-
-public:
-
-	UPROPERTY(EditAnywhere, BlueprintReadWrite, meta=(ExposeOnSpawn), Category = "ImageSequenceInfo")
-	FTimecode Timecode;
-
-	UPROPERTY(EditAnywhere, BlueprintReadWrite, meta=(ExposeOnSpawn), Category = "ImageSequenceInfo")
-	FFrameRate FrameRate;
-};
-
 UCLASS(BlueprintType, Blueprintable)
 class CAPTUREDATAUTILS_API UImageSequenceTimecodeUtils
 	: public UObject
diff --git a/Engine/Plugins/VirtualProduction/CaptureData/Source/CaptureDataUtils/Public/SoundWaveTimecodeUtils.h b/Engine/Plugins/VirtualProduction/CaptureData/Source/CaptureDataUtils/Public/SoundWaveTimecodeUtils.h
new file mode 100644
--- /dev/null
+++ b/Engine/Plugins/VirtualProduction/CaptureData/Source/CaptureDataUtils/Public/SoundWaveTimecodeUtils.h
@@ -0,0 +1,27 @@
+// Copyright Epic Games, Inc. All Rights Reserved.
+
+#pragma once
+
+#include "UObject/Object.h"
+#include "Misc/Timecode.h"
+#include "Misc/FrameRate.h"
+#include "Sound/SoundWave.h"
+
+#include "SoundWaveTimecodeUtils.generated.h"
+
+UCLASS(BlueprintType, Blueprintable)
+class CAPTUREDATAUTILS_API USoundWaveTimecodeUtils : public UObject
+{
+	GENERATED_BODY()
+
+public:
+
+	UFUNCTION(BlueprintCallable, Category = "Audio")
+	static void SetTimecodeInfo(const FTimecode& InTimecode, const FFrameRate& InFrameRate, USoundWave* OutSoundWave);
+
+	UFUNCTION(BlueprintCallable, Category = "Audio")
+	static FTimecode GetTimecode(const USoundWave* InSoundWave);
+
+	UFUNCTION(BlueprintCallable, Category = "Audio")
+	static FFrameRate GetFrameRate(const USoundWave* InSoundWave);
+};
