MetaHuman:  Fixed SoundWave always displaying drop frame timecode.

diff --git a/Engine/Plugins/VirtualProduction/CaptureData/Source/CaptureDataUtils/Private/SoundWaveTimecodeUtils.cpp b/Engine/Plugins/VirtualProduction/CaptureData/Source/CaptureDataUtils/Private/SoundWaveTimecodeUtils.cpp
--- a/Engine/Plugins/VirtualProduction/CaptureData/Source/CaptureDataUtils/Private/SoundWaveTimecodeUtils.cpp
+++ b/Engine/Plugins/VirtualProduction/CaptureData/Source/CaptureDataUtils/Private/SoundWaveTimecodeUtils.cpp
@@ -64,7 +64,10 @@ FTimecode USoundWaveTimecodeUtils::GetTimecode(const USoundWave* InSoundWave)
 	// GetNumSecondsSinceMidnight will return 0.0 in this case, so nothing here will fail, but it may be useful to fail early during dev
 	check(SoundWaveTimecodeInfo->NumSamplesPerSecond > 0);
 
-	FTimecode Timecode(SoundWaveTimecodeInfo->GetNumSecondsSinceMidnight(), SoundWaveTimecodeInfo->TimecodeRate, true);
+	constexpr bool bRollover = true;
+	FTimecode Timecode(SoundWaveTimecodeInfo->GetNumSecondsSinceMidnight(), SoundWaveTimecodeInfo->TimecodeRate, bRollover);
+	Timecode.bDropFrameFormat = SoundWaveTimecodeInfo->bTimecodeIsDropFrame;
+
 	return Timecode;
 }
 
