Using StartTimecode property from the ImageSequence

diff --git a/Engine/Plugins/VirtualProduction/CaptureData/Source/CaptureDataCore/Private/CaptureData.cpp b/Engine/Plugins/VirtualProduction/CaptureData/Source/CaptureDataCore/Private/CaptureData.cpp
--- a/Engine/Plugins/VirtualProduction/CaptureData/Source/CaptureDataCore/Private/CaptureData.cpp
+++ b/Engine/Plugins/VirtualProduction/CaptureData/Source/CaptureDataCore/Private/CaptureData.cpp
@@ -22,9 +22,9 @@
 #include "MediaTexture.h"
 #include "Algo/AllOf.h"
 #include "Engine/SkeletalMesh.h"
+#include "ImageSequenceTimecodeUtils.h"
 
 #if WITH_EDITOR
-#include "ImageSequenceTimecodeUtils.h"
 #include "IDirectoryWatcher.h"
 #include "DirectoryWatcherModule.h"
 #endif
@@ -488,13 +488,11 @@ FTimecode UFootageCaptureData::GetEffectiveImageTimecode(int32 InView) const
 	check(InView < ImageSequences.Num());
 	check(InView < DepthSequences.Num());
 
-#if WITH_EDITOR
-	TOptional<FTimecode> TimecodeOpt = UImageSequenceTimecodeUtils::TryGetTimecode(ImageSequences[InView]);
-	if (TimecodeOpt.IsSet())
+	FTimecode ImgSequenceTimecode = UImageSequenceTimecodeUtils::GetTimecode(ImageSequences[InView]);
+	if (UImageSequenceTimecodeUtils::IsValidTimecode(ImgSequenceTimecode))
 	{
-		return MoveTemp(TimecodeOpt.GetValue());
+		return ImgSequenceTimecode;
 	}
-#endif
 
 	FTimecode Timecode;
 	FFrameRate FrameRate;
@@ -508,13 +506,11 @@ FFrameRate UFootageCaptureData::GetEffectiveImageTimecodeRate(int32 InView) cons
 	check(InView < ImageSequences.Num());
 	check(InView < DepthSequences.Num());
 
-#if WITH_EDITOR
-	TOptional<FFrameRate> FrameRateOpt = UImageSequenceTimecodeUtils::TryGetFrameRate(ImageSequences[InView]);
-	if (FrameRateOpt.IsSet())
+	FFrameRate ImgSequenceFrameRate = UImageSequenceTimecodeUtils::GetFrameRate(ImageSequences[InView]);
+	if (UImageSequenceTimecodeUtils::IsValidFrameRate(ImgSequenceFrameRate))
 	{
-		return MoveTemp(FrameRateOpt.GetValue());
+		return ImgSequenceFrameRate;
 	}
-#endif
 
 	FTimecode Timecode;
 	FFrameRate FrameRate;
@@ -528,13 +524,11 @@ FTimecode UFootageCaptureData::GetEffectiveDepthTimecode(int32 InView) const
 	check(InView < ImageSequences.Num());
 	check(InView < DepthSequences.Num());
 
-#if WITH_EDITOR
-	TOptional<FTimecode> TimecodeOpt = UImageSequenceTimecodeUtils::TryGetTimecode(DepthSequences[InView]);
-	if (TimecodeOpt.IsSet())
+	FTimecode DepthSequenceTimecode = UImageSequenceTimecodeUtils::GetTimecode(DepthSequences[InView]);
+	if (UImageSequenceTimecodeUtils::IsValidTimecode(DepthSequenceTimecode))
 	{
-		return MoveTemp(TimecodeOpt.GetValue());
+		return DepthSequenceTimecode;
 	}
-#endif
 
 	FTimecode Timecode;
 	FFrameRate FrameRate;
@@ -548,13 +542,12 @@ FFrameRate UFootageCaptureData::GetEffectiveDepthTimecodeRate(int32 InView) cons
 	check(InView < ImageSequences.Num());
 	check(InView < DepthSequences.Num());
 
-#if WITH_EDITOR
-	TOptional<FFrameRate> FrameRateOpt = UImageSequenceTimecodeUtils::TryGetFrameRate(DepthSequences[InView]);
-	if (FrameRateOpt.IsSet())
+	FFrameRate DepthSequenceFrameRate = UImageSequenceTimecodeUtils::GetFrameRate(DepthSequences[InView]);
+	if (UImageSequenceTimecodeUtils::IsValidFrameRate(DepthSequenceFrameRate))
 	{
-		return MoveTemp(FrameRateOpt.GetValue());
+		return DepthSequenceFrameRate;
 	}
-#endif
+
 	FTimecode Timecode;
 	FFrameRate FrameRate;
 	GetDefaultTimecodeInfo(Timecode, FrameRate);
@@ -618,34 +611,31 @@ FFrameRate UFootageCaptureData::GetEffectiveAudioTimecodeRate() const
 
 void UFootageCaptureData::GetDefaultTimecodeInfo(FTimecode& OutTimecode, FFrameRate& OutFrameRate) const
 {
-#if WITH_EDITOR
 	for (const TObjectPtr<class UImgMediaSource>& ImageSequence : ImageSequences)
 	{
-		TOptional<FTimecode> TimecodeOpt = UImageSequenceTimecodeUtils::TryGetTimecode(ImageSequence);
-		TOptional<FFrameRate> FrameRateOpt = UImageSequenceTimecodeUtils::TryGetFrameRate(ImageSequence);
+		FTimecode Timecode = UImageSequenceTimecodeUtils::GetTimecode(ImageSequence);
+		FFrameRate FrameRate = UImageSequenceTimecodeUtils::GetFrameRate(ImageSequence);
 
-		if (TimecodeOpt.IsSet() && FrameRateOpt.IsSet())
+		if (UImageSequenceTimecodeUtils::IsValidTimecodeInfo(Timecode, FrameRate))
 		{
-			OutTimecode = MoveTemp(TimecodeOpt.GetValue());
-			OutFrameRate = MoveTemp(FrameRateOpt.GetValue());
+			OutTimecode = MoveTemp(Timecode);
+			OutFrameRate = MoveTemp(FrameRate);
 			return;
 		}
 	}
 
 	for (const TObjectPtr<class UImgMediaSource>& DepthSequence : DepthSequences)
 	{
-		TOptional<FTimecode> TimecodeOpt = UImageSequenceTimecodeUtils::TryGetTimecode(DepthSequence);
-		TOptional<FFrameRate> FrameRateOpt = UImageSequenceTimecodeUtils::TryGetFrameRate(DepthSequence);
+		FTimecode Timecode = UImageSequenceTimecodeUtils::GetTimecode(DepthSequence);
+		FFrameRate FrameRate = UImageSequenceTimecodeUtils::GetFrameRate(DepthSequence);
 
-		if (TimecodeOpt.IsSet() && FrameRateOpt.IsSet())
+		if (UImageSequenceTimecodeUtils::IsValidTimecodeInfo(Timecode, FrameRate))
 		{
-			OutTimecode = MoveTemp(TimecodeOpt.GetValue());
-			OutFrameRate = MoveTemp(FrameRateOpt.GetValue());
+			OutTimecode = MoveTemp(Timecode);
+			OutFrameRate = MoveTemp(FrameRate);
 			return;
 		}
 	}
-#endif
-
 
 	if (Audios.Num() > 0)
 	{
diff --git a/Engine/Plugins/VirtualProduction/CaptureData/Source/CaptureDataUtils/Private/ImageSequenceTimecodeUtils.cpp b/Engine/Plugins/VirtualProduction/CaptureData/Source/CaptureDataUtils/Private/ImageSequenceTimecodeUtils.cpp
--- a/Engine/Plugins/VirtualProduction/CaptureData/Source/CaptureDataUtils/Private/ImageSequenceTimecodeUtils.cpp
+++ b/Engine/Plugins/VirtualProduction/CaptureData/Source/CaptureDataUtils/Private/ImageSequenceTimecodeUtils.cpp
@@ -4,163 +4,51 @@
 
 #include "ParseTakeUtils.h"
 
-#if WITH_EDITOR
-#include "Subsystems/EditorAssetSubsystem.h"
-#include "Editor.h"
-#endif
-
-const FName UImageSequenceTimecodeUtils::TimecodeTagName = "Timecode";
-const FName UImageSequenceTimecodeUtils::TimecodeRateTagName = "TimecodeRate";
-
 void UImageSequenceTimecodeUtils::SetTimecodeInfo(const FTimecode& InTimecode, const FFrameRate& InFrameRate, UImgMediaSource* InImageSequence)
 {
-	FString ImageTimecodeString = InTimecode.ToString();
-	FString ImageTimecodeRateString = FString::SanitizeFloat(InFrameRate.AsDecimal());
-
-	SetTimecodeInfoString(ImageTimecodeString, ImageTimecodeRateString, InImageSequence);
+	InImageSequence->StartTimecode = InTimecode;
+	InImageSequence->FrameRateOverride = InFrameRate;
 }
 
 void UImageSequenceTimecodeUtils::SetTimecodeInfoString(const FString& InTimecode, const FString& InFrameRate, UImgMediaSource* InImageSequence)
 {
-#if WITH_EDITOR
-	UEditorAssetSubsystem* EditorAssetSubsystem = GEditor ? GEditor->GetEditorSubsystem<UEditorAssetSubsystem>() : nullptr;
-
-	if (!EditorAssetSubsystem)
-	{
-		return;
-	}
+	InImageSequence->StartTimecode = ParseTimecode(InTimecode);
 
-	EditorAssetSubsystem->SetMetadataTag(InImageSequence, UImageSequenceTimecodeUtils::TimecodeTagName, InTimecode);
-	EditorAssetSubsystem->SetMetadataTag(InImageSequence, UImageSequenceTimecodeUtils::TimecodeRateTagName, InFrameRate);
-#endif
+	double TimecodeRate = FCString::Atod(*InFrameRate);
+	InImageSequence->FrameRateOverride = ConvertFrameRate(TimecodeRate);
 }
 
 FTimecode UImageSequenceTimecodeUtils::GetTimecode(UImgMediaSource* InImageSequence)
 {
-	TOptional<FTimecode> TimecodeOpt = TryGetTimecode(InImageSequence);
-
-	if (TimecodeOpt.IsSet())
-	{
-		return TimecodeOpt.GetValue();
-	}
-
-	return FTimecode();
+	return InImageSequence->StartTimecode;
 }
 
 FFrameRate UImageSequenceTimecodeUtils::GetFrameRate(UImgMediaSource* InImageSequence)
 {
-	TOptional<FFrameRate> FrameRateOpt = TryGetFrameRate(InImageSequence);
-
-	if (FrameRateOpt.IsSet())
-	{
-		return FrameRateOpt.GetValue();
-	}
-
-	return FFrameRate();
+	return InImageSequence->FrameRateOverride;
 }
 
 FString UImageSequenceTimecodeUtils::GetTimecodeString(UImgMediaSource* InImageSequence)
 {
-	TOptional<FString> TimecodeOpt = TryGetTimecodeString(InImageSequence);
-
-	if (TimecodeOpt.IsSet())
-	{
-		return TimecodeOpt.GetValue();
-	}
-
-	return FString();
+	return InImageSequence->StartTimecode.ToString();
 }
 
 FString UImageSequenceTimecodeUtils::GetFrameRateString(UImgMediaSource* InImageSequence)
 {
-	TOptional<FString> FrameRateOpt = TryGetFrameRateString(InImageSequence);
-
-	if (FrameRateOpt.IsSet())
-	{
-		return FrameRateOpt.GetValue();
-	}
-
-	return FString();
+	return FString::SanitizeFloat(InImageSequence->FrameRateOverride.AsDecimal());
 }
 
-TOptional<FTimecode> UImageSequenceTimecodeUtils::TryGetTimecode(UImgMediaSource* InImageSequence)
+bool UImageSequenceTimecodeUtils::IsValidTimecodeInfo(const FTimecode& InTimecode, const FFrameRate& InTimecodeRate)
 {
-	TOptional<FString> TimecodeOpt = TryGetTimecodeString(InImageSequence);
-
-	if (!TimecodeOpt.IsSet())
-	{
-		return {};
-	}
-
-	return ParseTimecode(TimecodeOpt.GetValue());
+	return IsValidTimecode(InTimecode) && IsValidFrameRate(InTimecodeRate);
 }
 
-TOptional<FFrameRate> UImageSequenceTimecodeUtils::TryGetFrameRate(UImgMediaSource* InImageSequence)
+bool UImageSequenceTimecodeUtils::IsValidTimecode(const FTimecode& InTimecode)
 {
-	TOptional<FString> TimecodeRateOpt = TryGetFrameRateString(InImageSequence);
-
-	if (!TimecodeRateOpt.IsSet())
-	{
-		return {};
-	}
-
-	double TimecodeRate = FCString::Atod(*TimecodeRateOpt.GetValue());
-
-	return ConvertFrameRate(TimecodeRate);
-}
-
-TOptional<FString> UImageSequenceTimecodeUtils::TryGetTimecodeString(UImgMediaSource* InImageSequence)
-{
-#if WITH_EDITOR
-	if (!InImageSequence)
-	{
-		return {};
-	}
-
-	UEditorAssetSubsystem* EditorAssetSubsystem = GEditor ? GEditor->GetEditorSubsystem<UEditorAssetSubsystem>() : nullptr;
-
-	if (!EditorAssetSubsystem)
-	{
-		return {};
-	}
-
-	FString Timecode = EditorAssetSubsystem->GetMetadataTag(InImageSequence, UImageSequenceTimecodeUtils::TimecodeTagName);
-
-	if (Timecode.IsEmpty())
-	{
-		return {};
-	}
-
-	return Timecode;
-#else
-	return {};
-#endif
+	return InTimecode.IsValid() && InTimecode != FTimecode();
 }
 
-TOptional<FString> UImageSequenceTimecodeUtils::TryGetFrameRateString(UImgMediaSource* InImageSequence)
+bool UImageSequenceTimecodeUtils::IsValidFrameRate(const FFrameRate& InTimecodeRate)
 {
-#if WITH_EDITOR
-	if (!InImageSequence)
-	{
-		return {};
-	}
-
-	UEditorAssetSubsystem* EditorAssetSubsystem = GEditor ? GEditor->GetEditorSubsystem<UEditorAssetSubsystem>() : nullptr;
-
-	if (!EditorAssetSubsystem)
-	{
-		return {};
-	}
-
-	FString TimecodeRateMetadata = EditorAssetSubsystem->GetMetadataTag(InImageSequence, UImageSequenceTimecodeUtils::TimecodeRateTagName);
-
-	if (!TimecodeRateMetadata.IsNumeric())
-	{
-		return {};
-	}
-	
-	return TimecodeRateMetadata;
-#else
-	return {};
-#endif
+	return InTimecodeRate.IsValid() && InTimecodeRate != FFrameRate();
 }
diff --git a/Engine/Plugins/VirtualProduction/CaptureData/Source/CaptureDataUtils/Public/ImageSequenceTimecodeUtils.h b/Engine/Plugins/VirtualProduction/CaptureData/Source/CaptureDataUtils/Public/ImageSequenceTimecodeUtils.h
--- a/Engine/Plugins/VirtualProduction/CaptureData/Source/CaptureDataUtils/Public/ImageSequenceTimecodeUtils.h
+++ b/Engine/Plugins/VirtualProduction/CaptureData/Source/CaptureDataUtils/Public/ImageSequenceTimecodeUtils.h
@@ -36,8 +36,7 @@ public:
 	UFUNCTION(BlueprintCallable, Category = "ImageSequence")
 	static FString GetFrameRateString(UImgMediaSource* InImageSequence);
 
-	static TOptional<FTimecode> TryGetTimecode(UImgMediaSource* InImageSequence);
-	static TOptional<FFrameRate> TryGetFrameRate(UImgMediaSource* InImageSequence);
-	static TOptional<FString> TryGetTimecodeString(UImgMediaSource* InImageSequence);
-	static TOptional<FString> TryGetFrameRateString(UImgMediaSource* InImageSequence);
+	static bool IsValidTimecodeInfo(const FTimecode& InTimecode, const FFrameRate& InTimecodeRate);
+	static bool IsValidTimecode(const FTimecode& InTimecode);
+	static bool IsValidFrameRate(const FFrameRate& InTimecodeRate);
 };
