Integrate Capture Data plugin into Capture Manager plugin

diff --git a/Engine/Plugins/VirtualProduction/CaptureData/Source/CaptureDataCore/CaptureDataCore.Build.cs b/Engine/Plugins/VirtualProduction/CaptureData/Source/CaptureDataCore/CaptureDataCore.Build.cs
--- a/Engine/Plugins/VirtualProduction/CaptureData/Source/CaptureDataCore/CaptureDataCore.Build.cs
+++ b/Engine/Plugins/VirtualProduction/CaptureData/Source/CaptureDataCore/CaptureDataCore.Build.cs
@@ -10,13 +10,14 @@ public class CaptureDataCore : ModuleRules
 
 		PublicDependencyModuleNames.AddRange(new string[]
 		{
-
+			"Media"
 		});
 
 		PrivateDependencyModuleNames.AddRange(new string[]
 		{
 			"Core",
 			"CoreUObject",
+			"Engine",
 			"ImgMedia",
 			"CameraCalibrationCore",
 			"MeshDescription",
@@ -25,4 +26,4 @@ public class CaptureDataCore : ModuleRules
 			"CaptureDataUtils"
 		});
 	}
-}
+}
\ No newline at end of file
diff --git a/Engine/Plugins/VirtualProduction/CaptureData/Source/CaptureDataCore/Private/CameraCalibration.cpp b/Engine/Plugins/VirtualProduction/CaptureData/Source/CaptureDataCore/Private/CameraCalibration.cpp
--- a/Engine/Plugins/VirtualProduction/CaptureData/Source/CaptureDataCore/Private/CameraCalibration.cpp
+++ b/Engine/Plugins/VirtualProduction/CaptureData/Source/CaptureDataCore/Private/CameraCalibration.cpp
@@ -68,7 +68,7 @@ void UCameraCalibration::GetAssetRegistryTags(TArray<FAssetRegistryTag>& OutTags
 }
 
 bool UCameraCalibration::ConvertToTrackerNodeCameraModels(TArray<FCameraCalibration>& OutCalibrations,
-																   TArray<TPair<FString, FString>>& OutStereoReconstructionPairs) const
+														  TArray<TPair<FString, FString>>& OutStereoReconstructionPairs) const
 {
 	OutCalibrations.SetNum(CameraCalibrations.Num());
 	OutStereoReconstructionPairs.SetNum(StereoPairs.Num());
@@ -83,6 +83,7 @@ bool UCameraCalibration::ConvertToTrackerNodeCameraModels(TArray<FCameraCalibrat
 	{
 		FCameraCalibration CurCalib;
 		CurCalib.Name = CameraCalibrations[Cam].Name;
+		CurCalib.Camera = CameraCalibrations[Cam].Name;
 		CurCalib.ImageSizeX = CameraCalibrations[Cam].LensFile->LensInfo.ImageDimensions.X;
 		CurCalib.ImageSizeY = CameraCalibrations[Cam].LensFile->LensInfo.ImageDimensions.Y;
 
@@ -124,6 +125,7 @@ bool UCameraCalibration::ConvertToTrackerNodeCameraModels(TArray<FCameraCalibrat
 		const FVector2D& FxFy = FocalLengthData.FxFy;
 		CurCalib.FX = FxFy.X * CurCalib.ImageSizeX;
 		CurCalib.FY = FxFy.Y * CurCalib.ImageSizeY;
+		CurCalib.FocalLength = FxFy;
 
 		const FImageCenterTable& ImageCenterTable = CameraCalibrations[Cam].LensFile->ImageCenterTable;
 		FImageCenterInfo ImageCenterData;
@@ -247,6 +249,101 @@ bool UCameraCalibration::ConvertFromTrackerNodeCameraModels(TArray<FCameraCalibr
 	return true;
 }
 
+bool UCameraCalibration::ConvertFromTrackerNodeCameraModelsV2(const TArray<FCameraCalibration>& InCalibrations)
+{
+	for (int32 Index = 0; Index < InCalibrations.Num(); ++Index)
+	{
+		const FCameraCalibration& Calibration = InCalibrations[Index];
+
+		const FString ObjectName = Calibration.bIsDepth ?
+			FString::Printf(TEXT("%s_Depth_LensFile"), *GetName()) :
+			FString::Printf(TEXT("%s_%s_RGB_LensFile"), *GetName(), *Calibration.Camera);
+		FString ParentPath = FString::Printf(TEXT("%s/../%s"), *GetPackage()->GetPathName(), *ObjectName);
+		FPaths::CollapseRelativeDirectories(ParentPath);
+		UObject* Parent = CreatePackage(*ParentPath);
+		checkf(Parent, TEXT("Failed to create parent"));
+
+		FExtendedLensFile CameraCalibration;
+		CameraCalibration.Name = Calibration.Camera;
+		CameraCalibration.IsDepthCamera = Calibration.bIsDepth;
+		CameraCalibration.LensFile = NewObject<ULensFile>(Parent, ULensFile::StaticClass(), *ObjectName, GetFlags());
+
+		// These a for a non-FIZ camera.
+		const float Focus = 0.0f;
+		const float Zoom = 0.0f;
+
+		// LensInfo
+		CameraCalibration.LensFile->LensInfo.LensModel = USphericalLensModel::StaticClass();
+		CameraCalibration.LensFile->LensInfo.LensModelName = FString::Printf(TEXT("Lens"));
+		// lens serial number is not needed
+
+		// leave sensor dimensions with default values and de-normalize using VideoDimensions or DepthDimensions
+		CameraCalibration.LensFile->LensInfo.ImageDimensions = FIntPoint(Calibration.ImageSizeX, Calibration.ImageSizeY);
+
+		// FocalLengthInfo
+		FFocalLengthInfo FocalLengthInfo;
+		FocalLengthInfo.FxFy = Calibration.FocalLength;
+
+		// DistortionInfo
+		FDistortionInfo DistortionInfo;
+		FSphericalDistortionParameters SphericalParameters;
+
+		SphericalParameters.K1 = Calibration.K1;
+		SphericalParameters.K2 = Calibration.K2;
+		SphericalParameters.P1 = Calibration.P1;
+		SphericalParameters.P2 = Calibration.P2;
+		SphericalParameters.K3 = Calibration.K3;
+
+		USphericalLensModel::StaticClass()->GetDefaultObject<ULensModel>()->ToArray(
+			SphericalParameters,
+			DistortionInfo.Parameters
+		);
+
+		// Normalized ImageCenterInfo
+		FImageCenterInfo ImageCenterInfo;
+		ImageCenterInfo.PrincipalPoint = FVector2D(Calibration.CX, Calibration.CY);
+
+		// NodalOffset
+		FNodalPointOffset NodalPointOffset;
+		FTransform Transform;
+		Transform.SetFromMatrix(Calibration.Transform);
+		NodalPointOffset.LocationOffset = Transform.GetLocation();
+		NodalPointOffset.RotationOffset = Transform.GetRotation();
+
+		if (Calibration.Orientation == EMediaOrientation::CW90 ||
+			Calibration.Orientation == EMediaOrientation::CW270)
+		{
+			Swap(CameraCalibration.LensFile->LensInfo.ImageDimensions.X, CameraCalibration.LensFile->LensInfo.ImageDimensions.Y);
+			Swap(CameraCalibration.LensFile->LensInfo.SensorDimensions.X, CameraCalibration.LensFile->LensInfo.SensorDimensions.Y);
+			Swap(FocalLengthInfo.FxFy.X, FocalLengthInfo.FxFy.Y);
+
+			FVector2D UnrotatedPrincipalPoint = ImageCenterInfo.PrincipalPoint;
+			ImageCenterInfo.PrincipalPoint.X = 1.0 - UnrotatedPrincipalPoint.Y;
+			ImageCenterInfo.PrincipalPoint.Y = UnrotatedPrincipalPoint.X;
+		}
+
+		CameraCalibration.LensFile->AddDistortionPoint(Focus, Zoom, DistortionInfo, FocalLengthInfo);
+		CameraCalibration.LensFile->AddImageCenterPoint(Focus, Zoom, ImageCenterInfo);
+		CameraCalibration.LensFile->AddNodalOffsetPoint(Focus, Zoom, NodalPointOffset);
+
+		CameraCalibration.LensFile->MarkPackageDirty();
+		FAssetRegistryModule::AssetCreated(CameraCalibration.LensFile);
+
+		CameraCalibrations.Add(CameraCalibration);
+	}
+
+	if (InCalibrations.Num() == 3)
+	{
+		// stereo HMC so make the stereo pair
+		FStereoPair CameraPair;
+		CameraPair.CameraIndex1 = 0;
+		CameraPair.CameraIndex2 = 1;
+		StereoPairs.Add(CameraPair);
+	}
+
+	return true;
+}
+
 int32 UCameraCalibration::GetIndexByCameraName(const FString& InName) const
 {
 	for (int32 Calibration = 0; Calibration < CameraCalibrations.Num(); ++Calibration)
diff --git a/Engine/Plugins/VirtualProduction/CaptureData/Source/CaptureDataCore/Private/CaptureData.cpp b/Engine/Plugins/VirtualProduction/CaptureData/Source/CaptureDataCore/Private/CaptureData.cpp
--- a/Engine/Plugins/VirtualProduction/CaptureData/Source/CaptureDataCore/Private/CaptureData.cpp
+++ b/Engine/Plugins/VirtualProduction/CaptureData/Source/CaptureDataCore/Private/CaptureData.cpp
@@ -271,26 +271,37 @@ UFootageCaptureData::FVerifyResult UFootageCaptureData::ViewsContainsValidData()
 	CHECK_AND_PRINT_RET(!ImageSequences.IsEmpty(), TEXT("Capture Data doesn't contain image sequences"));
 	CHECK_AND_PRINT_RET(!DepthSequences.IsEmpty(), TEXT("Capture Data doesn't contain depth sequences"));
 
-	for (int32 Index = 0; Index < ImageSequences.Num(); Index++)
+	FFrameRate ImageFrameRate;
+	for (int32 Index = 0; Index < ImageSequences.Num(); ++Index)
 	{
 		int32 NumImageFrames = 0;
 		FIntVector2 ImageDimensions;
 
 		UImgMediaSource* ImageSequence = ImageSequences[Index];
-		UImgMediaSource* DepthSequence = DepthSequences[Index];
 
 		CHECK_AND_PRINT_RET(FImageSequenceUtils::GetImageSequenceInfo(ImageSequence, ImageDimensions, NumImageFrames), TEXT("Image Sequence asset is invalid"));
+		CHECK_AND_PRINT_RET(ImageSequence->FrameRateOverride.IsValid(), TEXT("Image Sequence asset contains invalid frame rate"));
+
+		ImageFrameRate = ImageSequence->FrameRateOverride;
+	}
+
+	FFrameRate DepthFrameRate;
+	for (int32 Index = 0; Index < DepthSequences.Num(); ++Index)
+	{
+		UImgMediaSource* DepthSequence = DepthSequences[Index];
 
 		int32 NumDepthFrames = 0;
 		FIntVector2 DepthDimensions;
 
-		CHECK_AND_PRINT_RET(FImageSequenceUtils::GetImageSequenceInfo(DepthSequence, DepthDimensions, NumImageFrames), TEXT("Depth Sequence asset is invalid"));
-
-		CHECK_AND_PRINT_RET(ImageSequence->FrameRateOverride.IsValid(), TEXT("Image Sequence asset contains invalid frame rate"));
+		CHECK_AND_PRINT_RET(FImageSequenceUtils::GetImageSequenceInfo(DepthSequence, DepthDimensions, NumDepthFrames), TEXT("Depth Sequence asset is invalid"));
 		CHECK_AND_PRINT_RET(DepthSequence->FrameRateOverride.IsValid(), TEXT("Depth Sequence asset contains invalid frame rate"));
-		CHECK_AND_PRINT_RET(ImageSequence->FrameRateOverride == DepthSequence->FrameRateOverride, TEXT("Image Sequence asset and Depth Sequence asset contain different frame rates"));
+
+		DepthFrameRate = DepthSequence->FrameRateOverride;
 	}
 
+	// It shouldn't be up to the Capture Data to prevent users from using the data with different frame rates.
+	// CHECK_AND_PRINT_RET(ImageFrameRate == DepthFrameRate, TEXT("Image Sequence asset and Depth Sequence asset contain different frame rates"));
+
 	return MakeValue();
 }
 
@@ -486,7 +497,6 @@ TRange<FFrameNumber> UFootageCaptureData::GetAudioFrameRange(const FFrameRate& I
 FTimecode UFootageCaptureData::GetEffectiveImageTimecode(int32 InView) const
 {
 	check(InView < ImageSequences.Num());
-	check(InView < DepthSequences.Num());
 
 	FTimecode ImgSequenceTimecode = UImageSequenceTimecodeUtils::GetTimecode(ImageSequences[InView]);
 	if (UImageSequenceTimecodeUtils::IsValidTimecode(ImgSequenceTimecode))
@@ -504,7 +514,6 @@ FTimecode UFootageCaptureData::GetEffectiveImageTimecode(int32 InView) const
 FFrameRate UFootageCaptureData::GetEffectiveImageTimecodeRate(int32 InView) const
 {
 	check(InView < ImageSequences.Num());
-	check(InView < DepthSequences.Num());
 
 	FFrameRate ImgSequenceFrameRate = UImageSequenceTimecodeUtils::GetFrameRate(ImageSequences[InView]);
 	if (UImageSequenceTimecodeUtils::IsValidFrameRate(ImgSequenceFrameRate))
@@ -521,7 +530,6 @@ FFrameRate UFootageCaptureData::GetEffectiveImageTimecodeRate(int32 InView) cons
 
 FTimecode UFootageCaptureData::GetEffectiveDepthTimecode(int32 InView) const
 {
-	check(InView < ImageSequences.Num());
 	check(InView < DepthSequences.Num());
 
 	FTimecode DepthSequenceTimecode = UImageSequenceTimecodeUtils::GetTimecode(DepthSequences[InView]);
@@ -539,7 +547,6 @@ FTimecode UFootageCaptureData::GetEffectiveDepthTimecode(int32 InView) const
 
 FFrameRate UFootageCaptureData::GetEffectiveDepthTimecodeRate(int32 InView) const
 {
-	check(InView < ImageSequences.Num());
 	check(InView < DepthSequences.Num());
 
 	FFrameRate DepthSequenceFrameRate = UImageSequenceTimecodeUtils::GetFrameRate(DepthSequences[InView]);
diff --git a/Engine/Plugins/VirtualProduction/CaptureData/Source/CaptureDataCore/Public/CameraCalibration.h b/Engine/Plugins/VirtualProduction/CaptureData/Source/CaptureDataCore/Public/CameraCalibration.h
--- a/Engine/Plugins/VirtualProduction/CaptureData/Source/CaptureDataCore/Public/CameraCalibration.h
+++ b/Engine/Plugins/VirtualProduction/CaptureData/Source/CaptureDataCore/Public/CameraCalibration.h
@@ -5,6 +5,8 @@
 #include "UObject/ObjectMacros.h"
 #include "EditorFramework/AssetImportData.h"
 
+#include "IMediaTextureSample.h"
+
 #include "CameraCalibration.generated.h"
 
 
@@ -42,8 +44,21 @@ public:
 struct CAPTUREDATACORE_API FCameraCalibration
 {
 public:
+
+	enum EType
+	{
+		Video = 0,
+		Depth,
+
+		Unknown,
+	};
+
 	FString Name;
 	FString Camera;
+	EType CameraType;
+
+	bool bIsDepth = false;
+
 	int32 ImageSizeX = 0;
 	int32 ImageSizeY = 0;
 	double FX = 0;
@@ -59,6 +74,10 @@ public:
 	double K5 = 0;
 	double K6 = 0;
 	FMatrix Transform;
+
+	// Already normalized
+	FVector2D FocalLength;
+	EMediaOrientation Orientation;
 };
 
 /** Camera Calibration Asset
@@ -95,6 +114,7 @@ public:
 										  TArray<TPair<FString, FString>>& OutStereoReconstructionPairs) const;
 
 	bool ConvertFromTrackerNodeCameraModels(TArray<FCameraCalibration>& InCalibrations);
+	bool ConvertFromTrackerNodeCameraModelsV2(const TArray<FCameraCalibration>& InCalibrations);
 
 	int32 GetIndexByCameraName(const FString& InName) const;
 };
\ No newline at end of file
