Fix crash when processing and exporting level sequence from audio

diff --git a/Engine/Plugins/VirtualProduction/CaptureData/Source/CaptureDataUtils/Private/ImageSequenceTimecodeUtils.cpp b/Engine/Plugins/VirtualProduction/CaptureData/Source/CaptureDataUtils/Private/ImageSequenceTimecodeUtils.cpp
--- a/Engine/Plugins/VirtualProduction/CaptureData/Source/CaptureDataUtils/Private/ImageSequenceTimecodeUtils.cpp
+++ b/Engine/Plugins/VirtualProduction/CaptureData/Source/CaptureDataUtils/Private/ImageSequenceTimecodeUtils.cpp
@@ -6,36 +6,58 @@
 
 void UImageSequenceTimecodeUtils::SetTimecodeInfo(const FTimecode& InTimecode, const FFrameRate& InFrameRate, UImgMediaSource* InImageSequence)
 {
-	InImageSequence->StartTimecode = InTimecode;
-	InImageSequence->FrameRateOverride = InFrameRate;
+	if (ensure(IsValid(InImageSequence)))
+	{
+		InImageSequence->StartTimecode = InTimecode;
+		InImageSequence->FrameRateOverride = InFrameRate;
+	}
 }
 
 void UImageSequenceTimecodeUtils::SetTimecodeInfoString(const FString& InTimecode, const FString& InFrameRate, UImgMediaSource* InImageSequence)
 {
-	InImageSequence->StartTimecode = ParseTimecode(InTimecode);
+	if (ensure(IsValid(InImageSequence)))
+	{
+		InImageSequence->StartTimecode = ParseTimecode(InTimecode);
 
-	double TimecodeRate = FCString::Atod(*InFrameRate);
-	InImageSequence->FrameRateOverride = ConvertFrameRate(TimecodeRate);
+		double TimecodeRate = FCString::Atod(*InFrameRate);
+		InImageSequence->FrameRateOverride = ConvertFrameRate(TimecodeRate);
+	}
 }
 
 FTimecode UImageSequenceTimecodeUtils::GetTimecode(UImgMediaSource* InImageSequence)
 {
-	return InImageSequence->StartTimecode;
+	if (ensure(IsValid(InImageSequence)))
+	{
+		return InImageSequence->StartTimecode;
+	}
+	return FTimecode();
 }
 
 FFrameRate UImageSequenceTimecodeUtils::GetFrameRate(UImgMediaSource* InImageSequence)
 {
-	return InImageSequence->FrameRateOverride;
+	if (ensure(IsValid(InImageSequence)))
+	{
+		return InImageSequence->FrameRateOverride;
+	}
+	return FFrameRate();
 }
 
 FString UImageSequenceTimecodeUtils::GetTimecodeString(UImgMediaSource* InImageSequence)
 {
-	return InImageSequence->StartTimecode.ToString();
+	if (ensure(IsValid(InImageSequence)))
+	{
+		return InImageSequence->StartTimecode.ToString();
+	}
+	return FString();
 }
 
 FString UImageSequenceTimecodeUtils::GetFrameRateString(UImgMediaSource* InImageSequence)
 {
-	return FString::SanitizeFloat(InImageSequence->FrameRateOverride.AsDecimal());
+	if (ensure(IsValid(InImageSequence)))
+	{
+		return FString::SanitizeFloat(InImageSequence->FrameRateOverride.AsDecimal());
+	}
+	return FString();
 }
 
 bool UImageSequenceTimecodeUtils::IsValidTimecodeInfo(const FTimecode& InTimecode, const FFrameRate& InTimecodeRate)
