Changes to the CameraCalibration struct

diff --git a/Engine/Plugins/VirtualProduction/CaptureData/Source/CaptureDataCore/Private/CameraCalibration.cpp b/Engine/Plugins/VirtualProduction/CaptureData/Source/CaptureDataCore/Private/CameraCalibration.cpp
--- a/Engine/Plugins/VirtualProduction/CaptureData/Source/CaptureDataCore/Private/CameraCalibration.cpp
+++ b/Engine/Plugins/VirtualProduction/CaptureData/Source/CaptureDataCore/Private/CameraCalibration.cpp
@@ -9,6 +9,7 @@
 #include "AssetRegistry/AssetRegistryModule.h"
 #include "UObject/Package.h"
 #include "Misc/Paths.h"
+#include "VectorTypes.h"
 
 void UCameraCalibration::PostInitProperties()
 {
@@ -78,14 +79,13 @@ bool UCameraCalibration::ConvertToTrackerNodeCameraModels(TArray<FCameraCalibrat
 		OutStereoReconstructionPairs[Pair] = TPair<FString, FString>(FString::FromInt(StereoPairs[Pair].CameraIndex1), FString::FromInt(StereoPairs[Pair].CameraIndex2));
 	}
 
-
 	for (int32 Cam = 0; Cam < CameraCalibrations.Num(); Cam++)
 	{
 		FCameraCalibration CurCalib;
-		CurCalib.Name = CameraCalibrations[Cam].Name;
-		CurCalib.Camera = CameraCalibrations[Cam].Name;
-		CurCalib.ImageSizeX = CameraCalibrations[Cam].LensFile->LensInfo.ImageDimensions.X;
-		CurCalib.ImageSizeY = CameraCalibrations[Cam].LensFile->LensInfo.ImageDimensions.Y;
+		CurCalib.CameraId = CameraCalibrations[Cam].Name;
+		CurCalib.CameraType = CameraCalibrations[Cam].IsDepthCamera ? FCameraCalibration::Depth : FCameraCalibration::Video;
+		CurCalib.ImageSize.X = CameraCalibrations[Cam].LensFile->LensInfo.ImageDimensions.X;
+		CurCalib.ImageSize.Y = CameraCalibrations[Cam].LensFile->LensInfo.ImageDimensions.Y;
 
 		if (CameraCalibrations[Cam].LensFile->LensInfo.LensModel != USphericalLensModel::StaticClass())
 		{
@@ -123,9 +123,8 @@ bool UCameraCalibration::ConvertToTrackerNodeCameraModels(TArray<FCameraCalibrat
 			return false;
 		}
 		const FVector2D& FxFy = FocalLengthData.FxFy;
-		CurCalib.FX = FxFy.X * CurCalib.ImageSizeX;
-		CurCalib.FY = FxFy.Y * CurCalib.ImageSizeY;
-		CurCalib.FocalLength = FxFy;
+		CurCalib.FocalLength = FxFy * CurCalib.ImageSize;
+		CurCalib.FocalLengthNormalized = FxFy;
 
 		const FImageCenterTable& ImageCenterTable = CameraCalibrations[Cam].LensFile->ImageCenterTable;
 		FImageCenterInfo ImageCenterData;
@@ -137,8 +136,8 @@ bool UCameraCalibration::ConvertToTrackerNodeCameraModels(TArray<FCameraCalibrat
 			return false;
 		}
 		const FVector2D& PrincipalPoint = ImageCenterData.PrincipalPoint;
-		CurCalib.CX = PrincipalPoint.X * CurCalib.ImageSizeX;
-		CurCalib.CY = PrincipalPoint.Y * CurCalib.ImageSizeY;
+		CurCalib.PrincipalPoint = PrincipalPoint * CurCalib.ImageSize;
+		CurCalib.PrincipalPointNormalized = PrincipalPoint;
 
 		const FNodalOffsetTable& NodalOffsetTable = CameraCalibrations[Cam].LensFile->NodalOffsetTable;
 		FNodalPointOffset NodalOffsetData;
@@ -155,6 +154,7 @@ bool UCameraCalibration::ConvertToTrackerNodeCameraModels(TArray<FCameraCalibrat
 		Transform.SetRotation(NodalOffsetData.RotationOffset);
 		FOpenCVHelperLocal::ConvertUnrealToOpenCV(Transform);
 		CurCalib.Transform = Transform.ToMatrixWithScale();
+		CurCalib.Orientation = CameraOrientation;
 
 		OutCalibrations[Cam] = CurCalib;
 	}
@@ -162,110 +162,23 @@ bool UCameraCalibration::ConvertToTrackerNodeCameraModels(TArray<FCameraCalibrat
 	return true;
 }
 
-bool UCameraCalibration::ConvertFromTrackerNodeCameraModels(TArray<FCameraCalibration>& InCalibrations)
-{
-	check(InCalibrations.Num() == 2 || InCalibrations.Num() == 3); // only works for stereo hmc, or for the case of one RGB stream and one depth stream
-
-	for (int32 Index = 0; Index < InCalibrations.Num(); ++Index)
-	{
-		const FCameraCalibration& Calibration = InCalibrations[Index];
-
-		bool bIsDepthCamera = (Index == InCalibrations.Num() - 1);
-
-		const FString ObjectName = bIsDepthCamera ?
-			FString::Printf(TEXT("%s_Depth_LensFile"), *GetName()) :
-			FString::Printf(TEXT("%s_%s_RGB_LensFile"), *GetName(), *Calibration.Name);
-		FString ParentPath = FString::Printf(TEXT("%s/../%s"), *GetPackage()->GetPathName(), *ObjectName);
-		FPaths::CollapseRelativeDirectories(ParentPath);
-		UObject* Parent = CreatePackage(*ParentPath);
-
-		FExtendedLensFile CameraCalibration;
-		CameraCalibration.Name = Calibration.Name;
-		CameraCalibration.IsDepthCamera = bIsDepthCamera;
-		CameraCalibration.LensFile = NewObject<ULensFile>(Parent, ULensFile::StaticClass(), *ObjectName, GetFlags());
-
-		// These a for a non-FIZ camera.
-		const float Focus = 0.0f;
-		const float Zoom = 0.0f;
-
-		// LensInfo
-		CameraCalibration.LensFile->LensInfo.LensModel = USphericalLensModel::StaticClass();
-		CameraCalibration.LensFile->LensInfo.LensModelName = FString::Printf(TEXT("Lens"));
-		// lens serial number is not needed
-		
-		// leave sensor dimensions with default values and de-normalize using VideoDimensions or DepthDimensions
-		CameraCalibration.LensFile->LensInfo.ImageDimensions = FIntPoint(Calibration.ImageSizeX, Calibration.ImageSizeY);
-
-		// FocalLengthInfo
-		FFocalLengthInfo FocalLengthInfo;
-		FocalLengthInfo.FxFy = FVector2D(Calibration.FX / Calibration.ImageSizeX, Calibration.FY / Calibration.ImageSizeY);
-
-		// DistortionInfo
-		FDistortionInfo DistortionInfo;
-		FSphericalDistortionParameters SphericalParameters;
-
-		SphericalParameters.K1 = Calibration.K1;
-		SphericalParameters.K2 = Calibration.K2;
-		SphericalParameters.P1 = Calibration.P1;
-		SphericalParameters.P2 = Calibration.P2;
-		SphericalParameters.K3 = Calibration.K3;
-
-		USphericalLensModel::StaticClass()->GetDefaultObject<ULensModel>()->ToArray(
-			SphericalParameters,
-			DistortionInfo.Parameters
-		);
-
-		// ImageCenterInfo
-		FImageCenterInfo ImageCenterInfo;
-		ImageCenterInfo.PrincipalPoint = FVector2D(Calibration.CX / Calibration.ImageSizeX, Calibration.CY / Calibration.ImageSizeY);
-
-		// NodalOffset
-		FNodalPointOffset NodalPointOffset;
-		FTransform Transform;
-		Transform.SetFromMatrix(Calibration.Transform);
-		FOpenCVHelperLocal::ConvertOpenCVToUnreal(Transform);
-		NodalPointOffset.LocationOffset = Transform.GetLocation();
-		NodalPointOffset.RotationOffset = Transform.GetRotation();
-
-		CameraCalibration.LensFile->AddDistortionPoint(Focus, Zoom, DistortionInfo, FocalLengthInfo);
-		CameraCalibration.LensFile->AddImageCenterPoint(Focus, Zoom, ImageCenterInfo);
-		CameraCalibration.LensFile->AddNodalOffsetPoint(Focus, Zoom, NodalPointOffset);
-
-		CameraCalibration.LensFile->MarkPackageDirty();
-		FAssetRegistryModule::AssetCreated(CameraCalibration.LensFile);
-
-		CameraCalibrations.Add(CameraCalibration);
-	}
-
-	if (InCalibrations.Num() == 3)
-	{
-		// stereo HMC so make the stereo pair
-		FStereoPair CameraPair;
-		CameraPair.CameraIndex1 = 0;
-		CameraPair.CameraIndex2 = 1;
-		StereoPairs.Add(CameraPair);
-	}
-
-	return true;
-}
-
-bool UCameraCalibration::ConvertFromTrackerNodeCameraModelsV2(const TArray<FCameraCalibration>& InCalibrations)
+bool UCameraCalibration::ConvertFromTrackerNodeCameraModels(const TArray<FCameraCalibration>& InCalibrations, bool bInUsingUnrealCoordinateSystem)
 {
 	for (int32 Index = 0; Index < InCalibrations.Num(); ++Index)
 	{
 		const FCameraCalibration& Calibration = InCalibrations[Index];
 
-		const FString ObjectName = Calibration.bIsDepth ?
+		const FString ObjectName = Calibration.CameraType == FCameraCalibration::Depth ?
 			FString::Printf(TEXT("%s_Depth_LensFile"), *GetName()) :
-			FString::Printf(TEXT("%s_%s_RGB_LensFile"), *GetName(), *Calibration.Camera);
+			FString::Printf(TEXT("%s_%s_RGB_LensFile"), *GetName(), *Calibration.CameraId);
 		FString ParentPath = FString::Printf(TEXT("%s/../%s"), *GetPackage()->GetPathName(), *ObjectName);
 		FPaths::CollapseRelativeDirectories(ParentPath);
 		UObject* Parent = CreatePackage(*ParentPath);
 		checkf(Parent, TEXT("Failed to create parent"));
 
 		FExtendedLensFile CameraCalibration;
-		CameraCalibration.Name = Calibration.Camera;
-		CameraCalibration.IsDepthCamera = Calibration.bIsDepth;
+		CameraCalibration.Name = Calibration.CameraId;
+		CameraCalibration.IsDepthCamera = Calibration.CameraType == FCameraCalibration::Depth;
 		CameraCalibration.LensFile = NewObject<ULensFile>(Parent, ULensFile::StaticClass(), *ObjectName, GetFlags());
 
 		// These a for a non-FIZ camera.
@@ -278,11 +191,18 @@ bool UCameraCalibration::ConvertFromTrackerNodeCameraModelsV2(const TArray<FCame
 		// lens serial number is not needed
 
 		// leave sensor dimensions with default values and de-normalize using VideoDimensions or DepthDimensions
-		CameraCalibration.LensFile->LensInfo.ImageDimensions = FIntPoint(Calibration.ImageSizeX, Calibration.ImageSizeY);
+		CameraCalibration.LensFile->LensInfo.ImageDimensions = FIntPoint(Calibration.ImageSize.X, Calibration.ImageSize.Y);
 
 		// FocalLengthInfo
 		FFocalLengthInfo FocalLengthInfo;
-		FocalLengthInfo.FxFy = Calibration.FocalLength;
+		if (!Calibration.FocalLengthNormalized.Equals(FVector2D::Zero()))
+		{
+			FocalLengthInfo.FxFy = Calibration.FocalLengthNormalized;
+		}
+		else
+		{
+			FocalLengthInfo.FxFy = Calibration.FocalLength / Calibration.ImageSize;
+		}
 
 		// DistortionInfo
 		FDistortionInfo DistortionInfo;
@@ -299,14 +219,26 @@ bool UCameraCalibration::ConvertFromTrackerNodeCameraModelsV2(const TArray<FCame
 			DistortionInfo.Parameters
 		);
 
-		// Normalized ImageCenterInfo
+		// ImageCenterInfo
 		FImageCenterInfo ImageCenterInfo;
-		ImageCenterInfo.PrincipalPoint = FVector2D(Calibration.CX, Calibration.CY);
+		if (!Calibration.PrincipalPointNormalized.Equals(FVector2D::Zero()))
+		{
+			ImageCenterInfo.PrincipalPoint = Calibration.PrincipalPointNormalized;
+		}
+		else
+		{
+			ImageCenterInfo.PrincipalPoint = Calibration.PrincipalPoint / Calibration.ImageSize;
+		}
 
 		// NodalOffset
 		FNodalPointOffset NodalPointOffset;
 		FTransform Transform;
 		Transform.SetFromMatrix(Calibration.Transform);
+		if (!bInUsingUnrealCoordinateSystem)
+		{
+			FOpenCVHelperLocal::ConvertOpenCVToUnreal(Transform);
+		}
+		
 		NodalPointOffset.LocationOffset = Transform.GetLocation();
 		NodalPointOffset.RotationOffset = Transform.GetRotation();
 
@@ -330,6 +262,7 @@ bool UCameraCalibration::ConvertFromTrackerNodeCameraModelsV2(const TArray<FCame
 		FAssetRegistryModule::AssetCreated(CameraCalibration.LensFile);
 
 		CameraCalibrations.Add(CameraCalibration);
+		CameraOrientation = Calibration.Orientation;
 	}
 
 	if (InCalibrations.Num() == 3)
@@ -344,18 +277,19 @@ bool UCameraCalibration::ConvertFromTrackerNodeCameraModelsV2(const TArray<FCame
 	return true;
 }
 
-int32 UCameraCalibration::GetIndexByCameraName(const FString& InName) const
+int32 UCameraCalibration::GetCalibrationIndexByName(const FString& InName) const
 {
-	for (int32 Calibration = 0; Calibration < CameraCalibrations.Num(); ++Calibration)
+	for (int32 Index = 0; Index < CameraCalibrations.Num(); ++Index)
 	{
-		if (CameraCalibrations[Calibration].Name == InName)
+		const FExtendedLensFile& Calibration = CameraCalibrations[Index];
+
+		if (Calibration.Name == InName)
 		{
-			return Calibration;
+			return Index;
 		}
 	}
 
 	UE_LOG(LogCaptureDataCore, Warning, TEXT("Specified camera name not valid"));
 
-	return -1;
-}
-
+	return INDEX_NONE;
+}
\ No newline at end of file
diff --git a/Engine/Plugins/VirtualProduction/CaptureData/Source/CaptureDataCore/Private/CaptureData.cpp b/Engine/Plugins/VirtualProduction/CaptureData/Source/CaptureDataCore/Private/CaptureData.cpp
--- a/Engine/Plugins/VirtualProduction/CaptureData/Source/CaptureDataCore/Private/CaptureData.cpp
+++ b/Engine/Plugins/VirtualProduction/CaptureData/Source/CaptureDataCore/Private/CaptureData.cpp
@@ -364,7 +364,15 @@ FIntPoint UFootageCaptureData::GetFootageColorResolution() const
 		TArray<TPair<FString, FString>> StereoPairs;
 		CameraCalibrations[0]->ConvertToTrackerNodeCameraModels(Calibrations, StereoPairs);
 
-		Resolution = FIntPoint(Calibrations[0].ImageSizeX, Calibrations[0].ImageSizeY);
+		FCameraCalibration* VideoCalibration = Calibrations.FindByPredicate([](const FCameraCalibration& InCalibration)
+		{
+			return InCalibration.CameraType == FCameraCalibration::Video;
+		});
+
+		if (VideoCalibration)
+		{
+			Resolution = FIntPoint(VideoCalibration->ImageSize.X, VideoCalibration->ImageSize.Y);
+		}
 	}
 
 	return Resolution;
@@ -753,4 +761,57 @@ void UFootageCaptureData::PopulateCameraNames(UFootageCaptureData* InFootageCapt
 			InOutCamera = *OutCameraNames[0];
 		}
 	}
+}
+
+int32 CheckCalibrationArray(const TArray<FExtendedLensFile>& InCalibrationArray, const FString& InName, bool bIsDepth)
+{
+	int32 Index = 0;
+	for (const FExtendedLensFile& Calibration : InCalibrationArray)
+	{
+		// Don't take into account cameras that do not match the camera type (is depth)
+		if (Calibration.IsDepthCamera != bIsDepth)
+		{
+			continue;
+		}
+
+		if (Calibration.Name == InName)
+		{
+			return Index;
+		}
+
+		// Only increase index for specified type of camera
+		++Index;
+	}
+
+	return INDEX_NONE;
+}
+
+int32 UFootageCaptureData::GetViewIndexByCameraName(const FString& InName) const
+{
+	if (CameraCalibrations.IsEmpty())
+	{
+		return INDEX_NONE;
+	}
+
+	ensure(CameraCalibrations.Num() == 1);
+
+	// Search for video camera
+	int32 Index = CheckCalibrationArray(CameraCalibrations[0]->CameraCalibrations, InName, false);
+
+	if (Index != INDEX_NONE && ImageSequences.IsValidIndex(Index))
+	{
+		return Index;
+	}
+
+	// Search for depth camera
+	Index = CheckCalibrationArray(CameraCalibrations[0]->CameraCalibrations, InName, true);
+
+	if (Index != INDEX_NONE && DepthSequences.IsValidIndex(Index))
+	{
+		return Index;
+	}
+
+	UE_LOG(LogCaptureDataCore, Warning, TEXT("Specified camera name not valid"));
+
+	return INDEX_NONE;
 }
\ No newline at end of file
diff --git a/Engine/Plugins/VirtualProduction/CaptureData/Source/CaptureDataCore/Public/CameraCalibration.h b/Engine/Plugins/VirtualProduction/CaptureData/Source/CaptureDataCore/Public/CameraCalibration.h
--- a/Engine/Plugins/VirtualProduction/CaptureData/Source/CaptureDataCore/Public/CameraCalibration.h
+++ b/Engine/Plugins/VirtualProduction/CaptureData/Source/CaptureDataCore/Public/CameraCalibration.h
@@ -52,19 +52,16 @@ public:
 
 		Unknown,
 	};
-
-	FString Name;
-	FString Camera;
+	
+	FString CameraId;
 	EType CameraType;
 
-	bool bIsDepth = false;
+	FVector2D ImageSize = FVector2D::Zero();
+	FVector2D FocalLength = FVector2D::Zero();
+	FVector2D PrincipalPoint = FVector2D::Zero();
+	FVector2D FocalLengthNormalized = FVector2D::Zero();
+	FVector2D PrincipalPointNormalized = FVector2D::Zero();
 
-	int32 ImageSizeX = 0;
-	int32 ImageSizeY = 0;
-	double FX = 0;
-	double FY = 0;
-	double CX = 0;
-	double CY = 0;
 	double K1 = 0;
 	double K2 = 0;
 	double P1 = 0;
@@ -74,10 +71,7 @@ public:
 	double K5 = 0;
 	double K6 = 0;
 	FMatrix Transform;
-
-	// Already normalized
-	FVector2D FocalLength;
-	EMediaOrientation Orientation;
+	EMediaOrientation Orientation = EMediaOrientation::Original;
 };
 
 /** Camera Calibration Asset
@@ -113,8 +107,11 @@ public:
 	bool ConvertToTrackerNodeCameraModels(TArray<FCameraCalibration>& OutCalibrations,
 										  TArray<TPair<FString, FString>>& OutStereoReconstructionPairs) const;
 
-	bool ConvertFromTrackerNodeCameraModels(TArray<FCameraCalibration>& InCalibrations);
-	bool ConvertFromTrackerNodeCameraModelsV2(const TArray<FCameraCalibration>& InCalibrations);
+	bool ConvertFromTrackerNodeCameraModels(const TArray<FCameraCalibration>& InCalibrations, bool bInUsingUnrealCoordinateSystem = false);
+
+	int32 GetCalibrationIndexByName(const FString& InName) const;
 
-	int32 GetIndexByCameraName(const FString& InName) const;
+private:
+	
+	EMediaOrientation CameraOrientation;
 };
\ No newline at end of file
diff --git a/Engine/Plugins/VirtualProduction/CaptureData/Source/CaptureDataCore/Public/CaptureData.h b/Engine/Plugins/VirtualProduction/CaptureData/Source/CaptureDataCore/Public/CaptureData.h
--- a/Engine/Plugins/VirtualProduction/CaptureData/Source/CaptureDataCore/Public/CaptureData.h
+++ b/Engine/Plugins/VirtualProduction/CaptureData/Source/CaptureDataCore/Public/CaptureData.h
@@ -278,6 +278,8 @@ public:
 	/** List of all RGB cameras (views) in the footage capture data. Ensures InOutCamera is a valid camera name. */
 	static void PopulateCameraNames(UFootageCaptureData* InFootageCaptureData, FString& InOutCamera, TArray<TSharedPtr<FString>>& OutCameraNames);
 
+	int32 GetViewIndexByCameraName(const FString& InName) const;
+
 private:
 	/** Checks if camera calibration and views are valid. */
 	FVerifyResult ViewsContainsValidData() const;
