Move Camera Calibration from top level content browser context menu into MetaHuman Animator section

diff --git a/Engine/Plugins/VirtualProduction/CaptureData/Source/CaptureDataEditor/Private/AssetDefinition_CameraCalibration.cpp b/Engine/Plugins/VirtualProduction/CaptureData/Source/CaptureDataEditor/Private/AssetDefinition_CameraCalibration.cpp
--- a/Engine/Plugins/VirtualProduction/CaptureData/Source/CaptureDataEditor/Private/AssetDefinition_CameraCalibration.cpp
+++ b/Engine/Plugins/VirtualProduction/CaptureData/Source/CaptureDataEditor/Private/AssetDefinition_CameraCalibration.cpp
@@ -20,7 +20,7 @@ TSoftClassPtr<UObject> UAssetDefinition_CameraCalibration::GetAssetClass() const
 
 TConstArrayView<FAssetCategoryPath> UAssetDefinition_CameraCalibration::GetAssetCategories() const
 {
-	static FAssetCategoryPath Path(NSLOCTEXT("CaptureData", "CaptureDataAdvancedAssetCategoryLabel", "Advanced"));
+	static FAssetCategoryPath Path(NSLOCTEXT("CaptureData", "CaptureDataAssetCategoryLabel", "MetaHuman Animator"), NSLOCTEXT("CaptureData", "CaptureDataAdvancedAssetCategoryLabel", "Advanced"));
 	static FAssetCategoryPath Categories[] = { Path };
 
 	return Categories;
