Streamline mandatory fields in input json file

diff --git a/Engine/Plugins/VirtualProduction/CaptureData/Source/CaptureDataCore/Private/CaptureData.cpp b/Engine/Plugins/VirtualProduction/CaptureData/Source/CaptureDataCore/Private/CaptureData.cpp
--- a/Engine/Plugins/VirtualProduction/CaptureData/Source/CaptureDataCore/Private/CaptureData.cpp
+++ b/Engine/Plugins/VirtualProduction/CaptureData/Source/CaptureDataCore/Private/CaptureData.cpp
@@ -307,7 +307,6 @@ UFootageCaptureData::FVerifyResult UFootageCaptureData::ViewsContainsValidData()
 
 UFootageCaptureData::FVerifyResult UFootageCaptureData::MetadataContainsValidData() const
 {
-	CHECK_AND_PRINT_RET(Metadata.Height > 0 && Metadata.Width > 0, TEXT("Invalid values for the image resolution"));
 	CHECK_AND_PRINT_RET(!FMath::IsNearlyZero(Metadata.FrameRate), TEXT("Frame rate can't be set to 0"));
 
 	return MakeValue();
diff --git a/Engine/Plugins/VirtualProduction/CaptureData/Source/CaptureDataCore/Public/CaptureData.h b/Engine/Plugins/VirtualProduction/CaptureData/Source/CaptureDataCore/Public/CaptureData.h
--- a/Engine/Plugins/VirtualProduction/CaptureData/Source/CaptureDataCore/Public/CaptureData.h
+++ b/Engine/Plugins/VirtualProduction/CaptureData/Source/CaptureDataCore/Public/CaptureData.h
@@ -100,11 +100,11 @@ struct CAPTUREDATACORE_API FFootageCaptureMetadata
 {
 	GENERATED_BODY()
 
-	UPROPERTY(EditAnywhere, BlueprintReadWrite, Category = "Metadata")
-	int32 Width = 0;
+	UPROPERTY(meta = (DeprecatedProperty, DeprecationMessage = "Redundant property"))
+	int32 Width_DEPRECATED = 0;
 
-	UPROPERTY(EditAnywhere, BlueprintReadWrite, Category = "Metadata")
-	int32 Height = 0;
+	UPROPERTY(meta = (DeprecatedProperty, DeprecationMessage = "Redundant property"))
+	int32 Height_DEPRECATED = 0;
 
 	UPROPERTY(EditAnywhere, BlueprintReadWrite, Category = "Metadata")
 	double FrameRate = 0.0;
