Make depth sequence optional in capture data and processing (part1)

diff --git a/Engine/Plugins/VirtualProduction/CaptureData/Source/CaptureDataCore/Private/CaptureData.cpp b/Engine/Plugins/VirtualProduction/CaptureData/Source/CaptureDataCore/Private/CaptureData.cpp
--- a/Engine/Plugins/VirtualProduction/CaptureData/Source/CaptureDataCore/Private/CaptureData.cpp
+++ b/Engine/Plugins/VirtualProduction/CaptureData/Source/CaptureDataCore/Private/CaptureData.cpp
@@ -58,7 +58,7 @@ void UCaptureData::NotifyInternalsChanged()
 /////////////////////////////////////////////////////
 // UMeshCaptureData
 
-bool UMeshCaptureData::IsInitialized() const
+bool UMeshCaptureData::IsInitialized(EInitializedCheck InInitializedCheck) const
 {
 	// We need access to the source data when fitting and that is not available during runtime
 #if WITH_EDITOR
@@ -237,16 +237,20 @@ void FFootageCaptureMetadata::SetDeviceClass(const FString& InDeviceModel)
 #define CHECK_AND_PRINT_RET(InCondition, InFormat, ...) if (!(InCondition)) { return MakeError(FString::Printf(InFormat, ##__VA_ARGS__)); }
 #define CHECK_AND_RET_ERROR(InFunction) if (UFootageCaptureData::FVerifyResult Result = InFunction; Result.HasError()) { return MoveTemp(Result); }
 
-bool UFootageCaptureData::IsInitialized() const
+bool UFootageCaptureData::IsInitialized(EInitializedCheck InInitializedCheck) const
 {
-	return !VerifyData().HasError();
+	return !VerifyData(InInitializedCheck).HasError();
 }
 
-UFootageCaptureData::FVerifyResult UFootageCaptureData::VerifyData() const
+UFootageCaptureData::FVerifyResult UFootageCaptureData::VerifyData(EInitializedCheck InInitializedCheck) const
 {
-	CHECK_AND_RET_ERROR(ViewsContainsValidData());
-	CHECK_AND_RET_ERROR(MetadataContainsValidData());
-	CHECK_AND_RET_ERROR(CalibrationContainsValidData());
+	CHECK_AND_RET_ERROR(ViewsContainsValidData(InInitializedCheck));
+
+	if (InInitializedCheck == EInitializedCheck::Full)
+	{
+		CHECK_AND_RET_ERROR(MetadataContainsValidData());
+		CHECK_AND_RET_ERROR(CalibrationContainsValidData());
+	}
 
 	return MakeValue();
 }
@@ -266,10 +270,9 @@ TArray<UFootageCaptureData::FPathAssociation> UFootageCaptureData::CheckImageSeq
 	return InvalidImageSequences;
 }
 
-UFootageCaptureData::FVerifyResult UFootageCaptureData::ViewsContainsValidData() const
+UFootageCaptureData::FVerifyResult UFootageCaptureData::ViewsContainsValidData(EInitializedCheck InInitializedCheck) const
 {
 	CHECK_AND_PRINT_RET(!ImageSequences.IsEmpty(), TEXT("Capture Data doesn't contain image sequences"));
-	CHECK_AND_PRINT_RET(!DepthSequences.IsEmpty(), TEXT("Capture Data doesn't contain depth sequences"));
 
 	FFrameRate ImageFrameRate;
 	for (int32 Index = 0; Index < ImageSequences.Num(); ++Index)
@@ -285,22 +288,27 @@ UFootageCaptureData::FVerifyResult UFootageCaptureData::ViewsContainsValidData()
 		ImageFrameRate = ImageSequence->FrameRateOverride;
 	}
 
-	FFrameRate DepthFrameRate;
-	for (int32 Index = 0; Index < DepthSequences.Num(); ++Index)
+	if (InInitializedCheck != EInitializedCheck::ImageSequencesOnly)
 	{
-		UImgMediaSource* DepthSequence = DepthSequences[Index];
+		CHECK_AND_PRINT_RET(!DepthSequences.IsEmpty(), TEXT("Capture Data doesn't contain depth sequences"));
+
+		FFrameRate DepthFrameRate;
+		for (int32 Index = 0; Index < DepthSequences.Num(); ++Index)
+		{
+			UImgMediaSource* DepthSequence = DepthSequences[Index];
 
-		int32 NumDepthFrames = 0;
-		FIntVector2 DepthDimensions;
+			int32 NumDepthFrames = 0;
+			FIntVector2 DepthDimensions;
 
-		CHECK_AND_PRINT_RET(FImageSequenceUtils::GetImageSequenceInfo(DepthSequence, DepthDimensions, NumDepthFrames), TEXT("Depth Sequence asset is invalid"));
-		CHECK_AND_PRINT_RET(DepthSequence->FrameRateOverride.IsValid(), TEXT("Depth Sequence asset contains invalid frame rate"));
+			CHECK_AND_PRINT_RET(FImageSequenceUtils::GetImageSequenceInfo(DepthSequence, DepthDimensions, NumDepthFrames), TEXT("Depth Sequence asset is invalid"));
+			CHECK_AND_PRINT_RET(DepthSequence->FrameRateOverride.IsValid(), TEXT("Depth Sequence asset contains invalid frame rate"));
 
-		DepthFrameRate = DepthSequence->FrameRateOverride;
-	}
+			DepthFrameRate = DepthSequence->FrameRateOverride;
+		}
 
-	// It shouldn't be up to the Capture Data to prevent users from using the data with different frame rates.
-	// CHECK_AND_PRINT_RET(ImageFrameRate == DepthFrameRate, TEXT("Image Sequence asset and Depth Sequence asset contain different frame rates"));
+		// It shouldn't be up to the Capture Data to prevent users from using the data with different frame rates.
+		// CHECK_AND_PRINT_RET(ImageFrameRate == DepthFrameRate, TEXT("Image Sequence asset and Depth Sequence asset contain different frame rates"));
+	}
 
 	return MakeValue();
 }
@@ -723,16 +731,26 @@ void UFootageCaptureData::PopulateCameraNames(UFootageCaptureData* InFootageCapt
 {
 	OutCameraNames.Reset();
 
-	if (InFootageCaptureData && !InFootageCaptureData->CameraCalibrations.IsEmpty())
+	if (InFootageCaptureData)
 	{
-		const UCameraCalibration* CameraCalibration = InFootageCaptureData->CameraCalibrations[0];
-		if (CameraCalibration)
+		if (InFootageCaptureData->CameraCalibrations.IsEmpty())
+		{
+			for (int32 Index = 0; Index < InFootageCaptureData->ImageSequences.Num(); ++Index)
+			{
+				OutCameraNames.Add(MakeShared<FString>(FString::Printf(TEXT("Camera %i"), Index))); // If you change the format also update UFootageCaptureData::GetViewIndexByCameraName
+			}
+		}
+		else
 		{
-			for (const FExtendedLensFile& LensFile : CameraCalibration->CameraCalibrations)
+			const UCameraCalibration* CameraCalibration = InFootageCaptureData->CameraCalibrations[0];
+			if (CameraCalibration)
 			{
-				if (!LensFile.IsDepthCamera)
+				for (const FExtendedLensFile& LensFile : CameraCalibration->CameraCalibrations)
 				{
-					OutCameraNames.Add(MakeShared<FString>(LensFile.Name));
+					if (!LensFile.IsDepthCamera)
+					{
+						OutCameraNames.Add(MakeShared<FString>(LensFile.Name));
+					}
 				}
 			}
 		}
@@ -789,7 +807,17 @@ int32 UFootageCaptureData::GetViewIndexByCameraName(const FString& InName) const
 {
 	if (CameraCalibrations.IsEmpty())
 	{
-		return INDEX_NONE;
+		int32 Index = INDEX_NONE;
+
+		TArray<FString> Tokens; // See UFootageCaptureData::PopulateCameraNames for camera name format
+		InName.ParseIntoArray(Tokens, TEXT(" "));
+
+		if (Tokens.Num() == 2 && Tokens[1].IsNumeric())
+		{
+			Index = FCString::Atoi(*Tokens[1]);
+		}
+
+		return Index;
 	}
 
 	ensure(CameraCalibrations.Num() == 1);
diff --git a/Engine/Plugins/VirtualProduction/CaptureData/Source/CaptureDataCore/Public/CaptureData.h b/Engine/Plugins/VirtualProduction/CaptureData/Source/CaptureDataCore/Public/CaptureData.h
--- a/Engine/Plugins/VirtualProduction/CaptureData/Source/CaptureDataCore/Public/CaptureData.h
+++ b/Engine/Plugins/VirtualProduction/CaptureData/Source/CaptureDataCore/Public/CaptureData.h
@@ -26,7 +26,14 @@ class UCaptureData : public UObject
 
 public:
 	/** Returns true is the capture data is fully initialized with all required information present */
-	virtual bool IsInitialized() const PURE_VIRTUAL(UCaptureData::IsInitialized, return false;);
+
+	enum EInitializedCheck
+	{
+		Full = 0,
+		ImageSequencesOnly,
+	};
+
+	virtual bool IsInitialized(EInitializedCheck InInitializedCheck) const PURE_VIRTUAL(UCaptureData::IsInitialized, return false;);
 
 	FOnCaptureDataInternalsChanged& OnCaptureDataInternalsChanged()
 	{
@@ -69,7 +76,7 @@ class CAPTUREDATACORE_API UMeshCaptureData : public UCaptureData
 
 public:
 	// UCaptureData interface
-	virtual bool IsInitialized() const override;
+	virtual bool IsInitialized(EInitializedCheck InInitializedCheck) const override;
 
 public:
 	/** Gets the data in the format expected by the face fitting API */
@@ -196,7 +203,7 @@ public:
 	using FVerifyResult = TValueOrError<void, FString>;
 
 	//~Begin UCaptureData interface
-	virtual bool IsInitialized() const override;
+	virtual bool IsInitialized(EInitializedCheck InInitializedCheck) const override;
 	//~End UCaptureData interface
 
 	//~Begin UObject interface
@@ -262,7 +269,7 @@ public:
 	FTimecode GetEffectiveAudioTimecode() const;
 	FFrameRate GetEffectiveAudioTimecodeRate() const;
 
-	FVerifyResult VerifyData() const;
+	FVerifyResult VerifyData(EInitializedCheck InInitializedCheck) const;
 
 	struct FPathAssociation
 	{
@@ -282,7 +289,7 @@ public:
 
 private:
 	/** Checks if camera calibration and views are valid. */
-	FVerifyResult ViewsContainsValidData() const;
+	FVerifyResult ViewsContainsValidData(EInitializedCheck InInitializedCheck) const;
 	FVerifyResult MetadataContainsValidData() const;
 	FVerifyResult CalibrationContainsValidData() const;
 
