Document DataIngestCore and DataIngestCoreEditor headers

diff --git a/Engine/Plugins/VirtualProduction/CaptureData/Source/CaptureDataCore/Private/CaptureData.cpp b/Engine/Plugins/VirtualProduction/CaptureData/Source/CaptureDataCore/Private/CaptureData.cpp
--- a/Engine/Plugins/VirtualProduction/CaptureData/Source/CaptureDataCore/Private/CaptureData.cpp
+++ b/Engine/Plugins/VirtualProduction/CaptureData/Source/CaptureDataCore/Private/CaptureData.cpp
@@ -282,7 +282,7 @@ UFootageCaptureData::FVerifyResult UFootageCaptureData::ViewsContainsValidData(E
 
 		UImgMediaSource* ImageSequence = ImageSequences[Index];
 
-		CHECK_AND_PRINT_RET(FImageSequenceUtils::GetImageSequenceInfo(ImageSequence, ImageDimensions, NumImageFrames), TEXT("Image Sequence asset is invalid"));
+		CHECK_AND_PRINT_RET(FImageSequenceUtils::GetImageSequenceInfoFromAsset(ImageSequence, ImageDimensions, NumImageFrames), TEXT("Image Sequence asset is invalid"));
 		CHECK_AND_PRINT_RET(ImageSequence->FrameRateOverride.IsValid(), TEXT("Image Sequence asset contains invalid frame rate"));
 
 		ImageFrameRate = ImageSequence->FrameRateOverride;
@@ -300,7 +300,7 @@ UFootageCaptureData::FVerifyResult UFootageCaptureData::ViewsContainsValidData(E
 			int32 NumDepthFrames = 0;
 			FIntVector2 DepthDimensions;
 
-			CHECK_AND_PRINT_RET(FImageSequenceUtils::GetImageSequenceInfo(DepthSequence, DepthDimensions, NumDepthFrames), TEXT("Depth Sequence asset is invalid"));
+			CHECK_AND_PRINT_RET(FImageSequenceUtils::GetImageSequenceInfoFromAsset(DepthSequence, DepthDimensions, NumDepthFrames), TEXT("Depth Sequence asset is invalid"));
 			CHECK_AND_PRINT_RET(DepthSequence->FrameRateOverride.IsValid(), TEXT("Depth Sequence asset contains invalid frame rate"));
 
 			DepthFrameRate = DepthSequence->FrameRateOverride;
@@ -348,7 +348,7 @@ static TRange<FFrameNumber> GetFrameRange(const FFrameRate& InTargetRate, UImgMe
 {
 	int32 Duration = 0;
 	FIntVector2 ImageDimensions;
-	const bool bImageOK = FImageSequenceUtils::GetImageSequenceInfo(InMedia, ImageDimensions, Duration);
+	const bool bImageOK = FImageSequenceUtils::GetImageSequenceInfoFromAsset(InMedia, ImageDimensions, Duration);
 	ensure(bImageOK);
 
 	return GetFrameRange(InTargetRate, InMediaTimecode, InMediaTimecodeRate, bInMediaStartFrameIsZero, InMedia->FrameRateOverride, Duration);
diff --git a/Engine/Plugins/VirtualProduction/CaptureData/Source/CaptureDataCore/Private/ImageSequenceUtils.cpp b/Engine/Plugins/VirtualProduction/CaptureData/Source/CaptureDataCore/Private/ImageSequenceUtils.cpp
--- a/Engine/Plugins/VirtualProduction/CaptureData/Source/CaptureDataCore/Private/ImageSequenceUtils.cpp
+++ b/Engine/Plugins/VirtualProduction/CaptureData/Source/CaptureDataCore/Private/ImageSequenceUtils.cpp
@@ -11,7 +11,7 @@
 #include "Modules/ModuleManager.h"
 
 
-bool FImageSequenceUtils::GetImageSequencePathAndFiles(const UImgMediaSource* InImgSequence, FString& OutFullSequencePath, TArray<FString>& OutImageFiles)
+bool FImageSequenceUtils::GetImageSequencePathAndFilesFromAsset(const UImgMediaSource* InImgSequence, FString& OutFullSequencePath, TArray<FString>& OutImageFiles)
 {
 	if (InImgSequence == nullptr)
 	{
@@ -20,10 +20,10 @@ bool FImageSequenceUtils::GetImageSequencePathAndFiles(const UImgMediaSource* In
 
 	OutFullSequencePath = InImgSequence->GetFullPath();
 
-	return GetImageSequencePathAndFiles(OutFullSequencePath, OutImageFiles);
+	return GetImageSequenceFilesFromPath(OutFullSequencePath, OutImageFiles);
 }
 
-bool FImageSequenceUtils::GetImageSequencePathAndFiles(const FString& InFullSequencePath, TArray<FString>& OutImageFiles)
+bool FImageSequenceUtils::GetImageSequenceFilesFromPath(const FString& InFullSequencePath, TArray<FString>& OutImageFiles)
 {
 	IImageWrapperModule& ImageWrapperModule = FModuleManager::LoadModuleChecked<IImageWrapperModule>(TEXT("ImageWrapper"));
 	IFileManager& FileManager = IFileManager::Get();
@@ -45,20 +45,20 @@ bool FImageSequenceUtils::GetImageSequencePathAndFiles(const FString& InFullSequ
 	return !OutImageFiles.IsEmpty() && bIterateResult;
 }
 
-bool FImageSequenceUtils::GetImageSequenceInfo(const class UImgMediaSource* InImgSequence, FIntVector2& OutDimensions, int32& OutNumImages)
+bool FImageSequenceUtils::GetImageSequenceInfoFromAsset(const class UImgMediaSource* InImgSequence, FIntVector2& OutDimensions, int32& OutNumImages)
 {
 	if (InImgSequence == nullptr)
 	{
 		return false;
 	}
 
-	return GetImageSequenceInfo(InImgSequence->GetFullPath(), OutDimensions, OutNumImages);
+	return GetImageSequenceInfoFromPath(InImgSequence->GetFullPath(), OutDimensions, OutNumImages);
 }
 
-bool FImageSequenceUtils::GetImageSequenceInfo(const FString& InFullSequencePath, FIntVector2& OutDimensions, int32& OutNumImages)
+bool FImageSequenceUtils::GetImageSequenceInfoFromPath(const FString& InFullSequencePath, FIntVector2& OutDimensions, int32& OutNumImages)
 {
 	TArray<FString> ImageFiles;
-	bool bFoundImages = GetImageSequencePathAndFiles(InFullSequencePath, ImageFiles);
+	bool bFoundImages = GetImageSequenceFilesFromPath(InFullSequencePath, ImageFiles);
 
 	if (!bFoundImages)
 	{
diff --git a/Engine/Plugins/VirtualProduction/CaptureData/Source/CaptureDataCore/Public/ImageSequenceUtils.h b/Engine/Plugins/VirtualProduction/CaptureData/Source/CaptureDataCore/Public/ImageSequenceUtils.h
--- a/Engine/Plugins/VirtualProduction/CaptureData/Source/CaptureDataCore/Public/ImageSequenceUtils.h
+++ b/Engine/Plugins/VirtualProduction/CaptureData/Source/CaptureDataCore/Public/ImageSequenceUtils.h
@@ -4,12 +4,19 @@
 
 #include "CoreMinimal.h"
 
+/** Utility functions for getting information about image sequences. */
 class CAPTUREDATACORE_API FImageSequenceUtils
 {
 public:
-	static bool GetImageSequencePathAndFiles(const class UImgMediaSource* InImgSequence, FString& OutFullSequencePath, TArray<FString>& OutImageFiles);
-	static bool GetImageSequencePathAndFiles(const FString& InFullSequencePath, TArray<FString>& OutImageFiles);
+	/** Get image sequence path and list of files from image media source asset. */
+	static bool GetImageSequencePathAndFilesFromAsset(const class UImgMediaSource* InImgSequence, FString& OutFullSequencePath, TArray<FString>& OutImageFiles);
 
-	static bool GetImageSequenceInfo(const class UImgMediaSource* InImgSequence, FIntVector2& OutDimensions, int32& OutNumImages);
-	static bool GetImageSequenceInfo(const FString& InFullSequencePath, FIntVector2& OutDimensions, int32& OutNumImages);
+	/** Get list of image file paths from the image sequence path. */
+	static bool GetImageSequenceFilesFromPath(const FString& InFullSequencePath, TArray<FString>& OutImageFiles);
+
+	/** Get image sequence info (e.g. dimensions, number of images) from image media source asset. */
+	static bool GetImageSequenceInfoFromAsset(const class UImgMediaSource* InImgSequence, FIntVector2& OutDimensions, int32& OutNumImages);
+
+	/** Get image sequence info (e.g. dimensions, number of images) from image sequence path. */
+	static bool GetImageSequenceInfoFromPath(const FString& InFullSequencePath, FIntVector2& OutDimensions, int32& OutNumImages);
 };
