MetaHuman: Resolved Horde test failures.

* Updated gold test data
* Updated static mesh
* Increased a few tolerances.

#rnx

#virtualized

diff --git a/Engine/Plugins/VirtualProduction/CaptureData/Source/CaptureDataCore/Public/CaptureData.h b/Engine/Plugins/VirtualProduction/CaptureData/Source/CaptureDataCore/Public/CaptureData.h
--- a/Engine/Plugins/VirtualProduction/CaptureData/Source/CaptureDataCore/Public/CaptureData.h
+++ b/Engine/Plugins/VirtualProduction/CaptureData/Source/CaptureDataCore/Public/CaptureData.h
@@ -227,7 +227,7 @@ public:
 	UPROPERTY(meta = (DeprecatedProperty, DeprecationMessage = "Audios is deprecated. Please use AudioTracks instead."))
 	TArray<TObjectPtr<class USoundWave>> Audios_DEPRECATED;
 
-	UPROPERTY(meta = (DeprecatedProperty, DeprecationMessage = "Audio is deprecated. Please use Audios instead."))
+	UPROPERTY(meta = (DeprecatedProperty, DeprecationMessage = "Audio is deprecated. Please use AudioTracks instead."))
 	TObjectPtr<class USoundWave> Audio_DEPRECATED;
 
 	UPROPERTY(meta = (DeprecatedProperty, DeprecationMessage = "AudioTimecodePresent is deprecated."))
