CaptureManager: Added missing dependency on EditorScriptingUtilities

diff --git a/Engine/Plugins/VirtualProduction/CaptureData/CaptureData.uplugin b/Engine/Plugins/VirtualProduction/CaptureData/CaptureData.uplugin
--- a/Engine/Plugins/VirtualProduction/CaptureData/CaptureData.uplugin
+++ b/Engine/Plugins/VirtualProduction/CaptureData/CaptureData.uplugin
@@ -39,6 +39,10 @@
 		{
 			"Name": "CameraCalibrationCore",
 			"Enabled": true
+		},
+		{
+			"Name": "EditorScriptingUtilities",
+			"Enabled":  true
 		}
 	]
 }
\ No newline at end of file
diff --git a/Engine/Plugins/VirtualProduction/CaptureData/Source/CaptureDataUtils/CaptureDataUtils.Build.cs b/Engine/Plugins/VirtualProduction/CaptureData/Source/CaptureDataUtils/CaptureDataUtils.Build.cs
--- a/Engine/Plugins/VirtualProduction/CaptureData/Source/CaptureDataUtils/CaptureDataUtils.Build.cs
+++ b/Engine/Plugins/VirtualProduction/CaptureData/Source/CaptureDataUtils/CaptureDataUtils.Build.cs
@@ -34,7 +34,8 @@ public class CaptureDataUtils : ModuleRules
 		{
 			PrivateDependencyModuleNames.AddRange(new string[]
 			{
-				"UnrealEd"
+				"UnrealEd",
+				"EditorScriptingUtilities"
 			});
 		}
 	}
