CaptureData bug fixes

These occured for me while I was trying to create identity. Unfortunately didn't succeed in recreating them. Locally, these make sense, but I didn't go full picture asking why some of these become nullptr at first place.

diff --git a/Engine/Plugins/VirtualProduction/CaptureData/Source/CaptureDataCore/Private/CaptureData.cpp b/Engine/Plugins/VirtualProduction/CaptureData/Source/CaptureDataCore/Private/CaptureData.cpp
--- a/Engine/Plugins/VirtualProduction/CaptureData/Source/CaptureDataCore/Private/CaptureData.cpp
+++ b/Engine/Plugins/VirtualProduction/CaptureData/Source/CaptureDataCore/Private/CaptureData.cpp
@@ -325,6 +325,7 @@ UFootageCaptureData::FVerifyResult UFootageCaptureData::CalibrationContainsValid
 	CHECK_AND_PRINT_RET(!CameraCalibrations.IsEmpty(), TEXT("Calibration assets are empty"));
 	for (UCameraCalibration* CameraCalibration : CameraCalibrations)
 	{
+		CHECK_AND_PRINT_RET(CameraCalibration != nullptr, TEXT("Calibration asset not configured"));
 		CHECK_AND_PRINT_RET(!CameraCalibration->CameraCalibrations.IsEmpty(), TEXT("Calibration asset not configured"));
 	}
 	return MakeValue();
@@ -585,13 +586,14 @@ FTimecode UFootageCaptureData::GetEffectiveAudioTimecode() const
 	if (AudioTracks.Num() > 0)
 	{
 		const USoundWave* SoundWave = AudioTracks[0].Get();
-		check(IsValid(SoundWave));
-
-		TOptional<FTimecode> TimecodeOpt = USoundWaveTimecodeUtils::GetTimecode(SoundWave);
-		if (TimecodeOpt.IsSet())
+		if (SoundWave != nullptr)
 		{
-			Timecode = TimecodeOpt.GetValue();
-			TimecodeFromAsset = true;
+			TOptional<FTimecode> TimecodeOpt = USoundWaveTimecodeUtils::GetTimecode(SoundWave);
+			if (TimecodeOpt.IsSet())
+			{
+				Timecode = TimecodeOpt.GetValue();
+				TimecodeFromAsset = true;
+			}
 		}
 	}
 
@@ -612,13 +614,14 @@ FFrameRate UFootageCaptureData::GetEffectiveAudioTimecodeRate() const
 	if (AudioTracks.Num() > 0)
 	{
 		const USoundWave* SoundWave = AudioTracks[0].Get();
-		check(IsValid(SoundWave));
-
-		TOptional<FFrameRate> FrameRateOpt = USoundWaveTimecodeUtils::GetFrameRate(SoundWave);
-		if (FrameRateOpt.IsSet())
+		if (SoundWave != nullptr)
 		{
-			FrameRate = FrameRateOpt.GetValue();
-			TimecodeRateFromAsset = true;
+			TOptional<FFrameRate> FrameRateOpt = USoundWaveTimecodeUtils::GetFrameRate(SoundWave);
+			if (FrameRateOpt.IsSet())
+			{
+				FrameRate = FrameRateOpt.GetValue();
+				TimecodeRateFromAsset = true;
+			}
 		}
 	}
 
@@ -635,43 +638,50 @@ void UFootageCaptureData::GetDefaultTimecodeInfo(FTimecode& OutTimecode, FFrameR
 {
 	for (const TObjectPtr<class UImgMediaSource>& ImageSequence : ImageSequences)
 	{
-		FTimecode Timecode = UImageSequenceTimecodeUtils::GetTimecode(ImageSequence);
-		FFrameRate FrameRate = UImageSequenceTimecodeUtils::GetFrameRate(ImageSequence);
-
-		if (UImageSequenceTimecodeUtils::IsValidTimecodeInfo(Timecode, FrameRate))
+		if (ImageSequence.Get() != nullptr)
 		{
-			OutTimecode = MoveTemp(Timecode);
-			OutFrameRate = MoveTemp(FrameRate);
-			return;
+			FTimecode Timecode = UImageSequenceTimecodeUtils::GetTimecode(ImageSequence);
+			FFrameRate FrameRate = UImageSequenceTimecodeUtils::GetFrameRate(ImageSequence);
+
+			if (UImageSequenceTimecodeUtils::IsValidTimecodeInfo(Timecode, FrameRate))
+			{
+				OutTimecode = MoveTemp(Timecode);
+				OutFrameRate = MoveTemp(FrameRate);
+				return;
+			}
 		}
 	}
 
 	for (const TObjectPtr<class UImgMediaSource>& DepthSequence : DepthSequences)
 	{
-		FTimecode Timecode = UImageSequenceTimecodeUtils::GetTimecode(DepthSequence);
-		FFrameRate FrameRate = UImageSequenceTimecodeUtils::GetFrameRate(DepthSequence);
-
-		if (UImageSequenceTimecodeUtils::IsValidTimecodeInfo(Timecode, FrameRate))
+		if (DepthSequence.Get() != nullptr)
 		{
-			OutTimecode = MoveTemp(Timecode);
-			OutFrameRate = MoveTemp(FrameRate);
-			return;
+			FTimecode Timecode = UImageSequenceTimecodeUtils::GetTimecode(DepthSequence);
+			FFrameRate FrameRate = UImageSequenceTimecodeUtils::GetFrameRate(DepthSequence);
+
+			if (UImageSequenceTimecodeUtils::IsValidTimecodeInfo(Timecode, FrameRate))
+			{
+				OutTimecode = MoveTemp(Timecode);
+				OutFrameRate = MoveTemp(FrameRate);
+				return;
+			}
 		}
 	}
 
 	if (AudioTracks.Num() > 0)
 	{
 		const USoundWave* SoundWave = AudioTracks[0].Get();
-		check(IsValid(SoundWave));
-
-		TOptional<FTimecode> TimecodeOpt = USoundWaveTimecodeUtils::GetTimecode(SoundWave);
-		TOptional<FFrameRate> FrameRateOpt = USoundWaveTimecodeUtils::GetFrameRate(SoundWave);
-
-		if (TimecodeOpt.IsSet() && FrameRateOpt.IsSet())
+		if (SoundWave != nullptr) 
 		{
-			OutTimecode = TimecodeOpt.GetValue();
-			OutFrameRate = FrameRateOpt.GetValue();
-			return;
+			TOptional<FTimecode> TimecodeOpt = USoundWaveTimecodeUtils::GetTimecode(SoundWave);
+			TOptional<FFrameRate> FrameRateOpt = USoundWaveTimecodeUtils::GetFrameRate(SoundWave);
+
+			if (TimecodeOpt.IsSet() && FrameRateOpt.IsSet())
+			{
+				OutTimecode = TimecodeOpt.GetValue();
+				OutFrameRate = FrameRateOpt.GetValue();
+				return;
+			}
 		}
 	}
 
