CaptureManager: Mark image sequence and audio assets as dirty after migration.

This is due to the fact that this is partially the case already for capture data and calibration assets, so this is needed to prevent reload issues when saving capture data assets post-migration.

#rnx

diff --git a/Engine/Plugins/VirtualProduction/CaptureData/Source/CaptureDataCore/Private/CaptureData.cpp b/Engine/Plugins/VirtualProduction/CaptureData/Source/CaptureDataCore/Private/CaptureData.cpp
--- a/Engine/Plugins/VirtualProduction/CaptureData/Source/CaptureDataCore/Private/CaptureData.cpp
+++ b/Engine/Plugins/VirtualProduction/CaptureData/Source/CaptureDataCore/Private/CaptureData.cpp
@@ -27,6 +27,7 @@
 #if WITH_EDITOR
 #include "IDirectoryWatcher.h"
 #include "DirectoryWatcherModule.h"
+#include "CaptureDataEditorBridge.h"
 #endif
 
 #include UE_INLINE_GENERATED_CPP_BY_NAME(CaptureData)
@@ -720,6 +721,8 @@ void UFootageCaptureData::PostLoad()
 		Audios_DEPRECATED.Empty();
 	}
 
+	TQueue<TWeakObjectPtr<UObject>> ObjectsToMarkDirty;
+
 #if WITH_EDITOR
 	for (FFootageCaptureView& View : Views_DEPRECATED)
 	{
@@ -727,12 +730,16 @@ void UFootageCaptureData::PostLoad()
 		{
 			UImageSequenceTimecodeUtils::SetTimecodeInfo(View.ImageTimecode, View.ImageTimecodeRate, View.ImageSequence.Get());
 		}
+
+		ObjectsToMarkDirty.Enqueue(View.ImageSequence);
 		ImageSequences.Add(MoveTemp(View.ImageSequence));
 
 		if (View.bDepthTimecodePresent)
 		{
 			UImageSequenceTimecodeUtils::SetTimecodeInfo(View.DepthTimecode, View.DepthTimecodeRate, View.DepthSequence.Get());
 		}
+
+		ObjectsToMarkDirty.Enqueue(View.DepthSequence);
 		DepthSequences.Add(MoveTemp(View.DepthSequence));
 	}
 #endif
@@ -743,7 +750,33 @@ void UFootageCaptureData::PostLoad()
 
 		USoundWaveTimecodeUtils::SetTimecodeInfo(AudioTimecode_DEPRECATED, AudioTimecodeRate_DEPRECATED, SoundWave);
 		bAudioTimecodePresent_DEPRECATED = false;
+		ObjectsToMarkDirty.Enqueue(SoundWave);
+	}
+
+#if WITH_EDITOR
+	// The calibration and capture data assets get marked dirty by the FCaptureDataCoreModule::CheckAssetMigration()
+	// function, so we need to mark the other assets we've modified here as dirty as well. Otherwise if a user closes
+	// the editor after this point they will be prompted to save only the calibration and capture data assets and not
+	// the image sequences, depth sequences etc., if that happens the timecode information we've migrated here will get
+	// lost, as those changes reside only in memory and the capture data PostLoad switches to update them won't get
+	// triggered again once the capture data asset is saved.
+	//
+	// We defer the mark until after PostLoad, to a point when the editor and relevant subsystems are known to be ready.
+
+	if (!ObjectsToMarkDirty.IsEmpty())
+	{
+		using namespace UE::CaptureManager;
+
+		ICaptureDataEditorBridge& EditorBridge = FModuleManager::LoadModuleChecked<ICaptureDataEditorBridge>("CaptureDataEditor");
+
+		TWeakObjectPtr<UObject> ObjectToMarkDirty;
+
+		while (ObjectsToMarkDirty.Dequeue(ObjectToMarkDirty))
+		{
+			EditorBridge.DeferMarkDirty(MoveTemp(ObjectToMarkDirty));
+		}
 	}
+#endif
 }
 
 void UFootageCaptureData::PopulateCameraNames(UFootageCaptureData* InFootageCaptureData, FString& InOutCamera, TArray<TSharedPtr<FString>>& OutCameraNames)
diff --git a/Engine/Plugins/VirtualProduction/CaptureData/Source/CaptureDataCore/Public/CaptureDataEditorBridge.h b/Engine/Plugins/VirtualProduction/CaptureData/Source/CaptureDataCore/Public/CaptureDataEditorBridge.h
new file mode 100644
--- /dev/null
+++ b/Engine/Plugins/VirtualProduction/CaptureData/Source/CaptureDataCore/Public/CaptureDataEditorBridge.h
@@ -0,0 +1,20 @@
+// Copyright Epic Games, Inc. All Rights Reserved.
+
+#pragma once
+
+#include "Modules/ModuleInterface.h"
+#include "UObject/WeakObjectPtrTemplates.h"
+
+class UObject;
+
+namespace UE::CaptureManager
+{
+
+class ICaptureDataEditorBridge : public IModuleInterface
+{
+public:
+	virtual void DeferMarkDirty(TWeakObjectPtr<UObject> InObject) = 0;
+};
+
+} // namespace UE::CaptureManager
+
diff --git a/Engine/Plugins/VirtualProduction/CaptureData/Source/CaptureDataEditor/CaptureDataEditor.Build.cs b/Engine/Plugins/VirtualProduction/CaptureData/Source/CaptureDataEditor/CaptureDataEditor.Build.cs
--- a/Engine/Plugins/VirtualProduction/CaptureData/Source/CaptureDataEditor/CaptureDataEditor.Build.cs
+++ b/Engine/Plugins/VirtualProduction/CaptureData/Source/CaptureDataEditor/CaptureDataEditor.Build.cs
@@ -25,6 +25,7 @@ public class CaptureDataEditor : ModuleRules
 			"SlateCore",
 			"Slate",
 			"InputCore",
+			"Engine"
 		});
 	}
 }
diff --git a/Engine/Plugins/VirtualProduction/CaptureData/Source/CaptureDataEditor/Private/CaptureDataEditorModule.cpp b/Engine/Plugins/VirtualProduction/CaptureData/Source/CaptureDataEditor/Private/CaptureDataEditorModule.cpp
--- a/Engine/Plugins/VirtualProduction/CaptureData/Source/CaptureDataEditor/Private/CaptureDataEditorModule.cpp
+++ b/Engine/Plugins/VirtualProduction/CaptureData/Source/CaptureDataEditor/Private/CaptureDataEditorModule.cpp
@@ -2,6 +2,8 @@
 
 #include "CaptureDataEditorModule.h"
 
+#include "DeferredObjectDirtier.h"
+
 #include "Modules/ModuleManager.h"
 
 void FCaptureDataEditorModule::StartupModule()
@@ -14,4 +16,12 @@ void FCaptureDataEditorModule::ShutdownModule()
 
 }
 
+void FCaptureDataEditorModule::DeferMarkDirty(TWeakObjectPtr<UObject> InObject)
+{
+	using namespace UE::CaptureManager;
+	
+	FDeferredObjectDirtier& ObjectDirtier = FDeferredObjectDirtier::Get();
+	ObjectDirtier.Enqueue(MoveTemp(InObject));
+}
+
 IMPLEMENT_MODULE(FCaptureDataEditorModule, CaptureDataEditor)
\ No newline at end of file
diff --git a/Engine/Plugins/VirtualProduction/CaptureData/Source/CaptureDataEditor/Private/CaptureDataEditorModule.h b/Engine/Plugins/VirtualProduction/CaptureData/Source/CaptureDataEditor/Private/CaptureDataEditorModule.h
--- a/Engine/Plugins/VirtualProduction/CaptureData/Source/CaptureDataEditor/Private/CaptureDataEditorModule.h
+++ b/Engine/Plugins/VirtualProduction/CaptureData/Source/CaptureDataEditor/Private/CaptureDataEditorModule.h
@@ -2,9 +2,9 @@
 
 #pragma once
 
-#include "Modules/ModuleInterface.h"
+#include "CaptureDataEditorBridge.h"
 
-class FCaptureDataEditorModule : public IModuleInterface
+class FCaptureDataEditorModule : public UE::CaptureManager::ICaptureDataEditorBridge
 {
 public:
 
@@ -13,5 +13,7 @@ public:
 	virtual void ShutdownModule() override;
 	//~ End IModuleInterface interface
 
-private:
+	// ~Begin ICaptureDataEditorBridge interface
+	virtual void DeferMarkDirty(TWeakObjectPtr<UObject> InObject) override;
+	// ~End ICaptureDataEditorBridge interface
 };
diff --git a/Engine/Plugins/VirtualProduction/CaptureData/Source/CaptureDataEditor/Private/DeferredObjectDirtier.cpp b/Engine/Plugins/VirtualProduction/CaptureData/Source/CaptureDataEditor/Private/DeferredObjectDirtier.cpp
new file mode 100644
--- /dev/null
+++ b/Engine/Plugins/VirtualProduction/CaptureData/Source/CaptureDataEditor/Private/DeferredObjectDirtier.cpp
@@ -0,0 +1,47 @@
+// Copyright Epic Games, Inc. All Rights Reserved.
+
+#include "DeferredObjectDirtier.h"
+
+#include "UObject/Object.h"
+
+namespace UE::CaptureManager
+{
+
+FDeferredObjectDirtier& FDeferredObjectDirtier::Get()
+{
+	static FDeferredObjectDirtier Instance;
+	return Instance;
+}
+
+void FDeferredObjectDirtier::Enqueue(TWeakObjectPtr<UObject> InObject)
+{
+	check(IsInGameThread());
+	ObjectsToMarkDirty.Emplace(MoveTemp(InObject));
+}
+
+void FDeferredObjectDirtier::Tick(const float InDeltaTime)
+{
+	for (TWeakObjectPtr<UObject>& ObjectToMarkDirty : ObjectsToMarkDirty)
+	{
+		if (ObjectToMarkDirty.IsValid())
+		{
+			ObjectToMarkDirty->Modify();
+			ObjectToMarkDirty->MarkPackageDirty();
+		}
+	}
+
+	ObjectsToMarkDirty.Empty();
+}
+
+bool FDeferredObjectDirtier::IsTickable() const
+{
+	return !ObjectsToMarkDirty.IsEmpty();
+}
+
+TStatId FDeferredObjectDirtier::GetStatId() const
+{
+	RETURN_QUICK_DECLARE_CYCLE_STAT(FDeferredObjectDirtier, STATGROUP_Tickables);
+}
+
+} // namespace UE::CaptureManager
+
diff --git a/Engine/Plugins/VirtualProduction/CaptureData/Source/CaptureDataEditor/Private/DeferredObjectDirtier.h b/Engine/Plugins/VirtualProduction/CaptureData/Source/CaptureDataEditor/Private/DeferredObjectDirtier.h
new file mode 100644
--- /dev/null
+++ b/Engine/Plugins/VirtualProduction/CaptureData/Source/CaptureDataEditor/Private/DeferredObjectDirtier.h
@@ -0,0 +1,28 @@
+// Copyright Epic Games, Inc. All Rights Reserved.
+
+#pragma once
+
+#include "TickableEditorObject.h"
+
+class UObject;
+
+namespace UE::CaptureManager
+{
+
+class FDeferredObjectDirtier : public FTickableEditorObject
+{
+public:
+	static FDeferredObjectDirtier& Get();
+	void Enqueue(TWeakObjectPtr<UObject> InObject);
+
+	virtual void Tick(float InDeltaTime) override;
+	virtual bool IsTickable() const override;
+	virtual TStatId GetStatId() const override;
+
+private:
+	FDeferredObjectDirtier() = default;
+	TArray<TWeakObjectPtr<UObject>> ObjectsToMarkDirty;
+};
+
+} // namespace UE::CaptureManager
+
