Adding Capture Metadata asset utility blueprint

#virtualized

diff --git a/Engine/Plugins/VirtualProduction/CaptureData/Source/CaptureDataEditor/CaptureDataEditor.Build.cs b/Engine/Plugins/VirtualProduction/CaptureData/Source/CaptureDataEditor/CaptureDataEditor.Build.cs
--- a/Engine/Plugins/VirtualProduction/CaptureData/Source/CaptureDataEditor/CaptureDataEditor.Build.cs
+++ b/Engine/Plugins/VirtualProduction/CaptureData/Source/CaptureDataEditor/CaptureDataEditor.Build.cs
@@ -25,7 +25,8 @@ public class CaptureDataEditor : ModuleRules
 			"SlateCore",
 			"Slate",
 			"InputCore",
-			"Engine"
+			"Engine",
+			"EditorScriptingUtilities"
 		});
 	}
 }
diff --git a/Engine/Plugins/VirtualProduction/CaptureData/Source/CaptureDataEditor/Private/CaptureDataEditorModule.cpp b/Engine/Plugins/VirtualProduction/CaptureData/Source/CaptureDataEditor/Private/CaptureDataEditorModule.cpp
--- a/Engine/Plugins/VirtualProduction/CaptureData/Source/CaptureDataEditor/Private/CaptureDataEditorModule.cpp
+++ b/Engine/Plugins/VirtualProduction/CaptureData/Source/CaptureDataEditor/Private/CaptureDataEditorModule.cpp
@@ -4,16 +4,24 @@
 
 #include "DeferredObjectDirtier.h"
 
+#include "PropertyEditorModule.h"
+
+#include "CaptureMetadata.h"
+#include "Customizations/CaptureMetadataCustomization.h"
+
 #include "Modules/ModuleManager.h"
 
 void FCaptureDataEditorModule::StartupModule()
 {
-
+	FPropertyEditorModule& PropertyEditorModule = FModuleManager::LoadModuleChecked<FPropertyEditorModule>("PropertyEditor");
+	PropertyEditorModule.RegisterCustomClassLayout(UCaptureMetadata::StaticClass()->GetFName(),
+												   FOnGetDetailCustomizationInstance::CreateStatic(&FCaptureMetadataCustomization::MakeInstance));
 }
 
 void FCaptureDataEditorModule::ShutdownModule()
 {
-
+	FPropertyEditorModule& PropertyEditorModule = FModuleManager::LoadModuleChecked<FPropertyEditorModule>("PropertyEditor");
+	PropertyEditorModule.UnregisterCustomClassLayout(UCaptureMetadata::StaticClass()->GetFName());
 }
 
 void FCaptureDataEditorModule::DeferMarkDirty(TWeakObjectPtr<UObject> InObject)
diff --git a/Engine/Plugins/VirtualProduction/CaptureData/Source/CaptureDataEditor/Private/CaptureMetadata.cpp b/Engine/Plugins/VirtualProduction/CaptureData/Source/CaptureDataEditor/Private/CaptureMetadata.cpp
new file mode 100644
--- /dev/null
+++ b/Engine/Plugins/VirtualProduction/CaptureData/Source/CaptureDataEditor/Private/CaptureMetadata.cpp
@@ -0,0 +1,183 @@
+// Copyright Epic Games, Inc. All Rights Reserved.
+
+#include "CaptureMetadata.h"
+
+#include "UObject/Package.h"
+
+#if WITH_EDITOR
+#include "EditorDialogLibrary.h"
+#endif
+
+void UCaptureMetadata::SetCaptureMetadata(UObject* InObject, const UCaptureMetadata* InCaptureMetadata)
+{
+	if (!IsValid(InObject))
+	{
+		return;
+	}
+
+	if (!IsValid(InCaptureMetadata))
+	{
+		return;
+	}
+
+	FMetaData& Metadata = InObject->GetPackage()->GetMetaData();
+
+	TFieldIterator<FProperty> It(InCaptureMetadata->GetClass());
+
+	bool bIsEdited = false;
+
+	while (It)
+	{
+		FProperty* Property = *It;
+		check(Property->IsA<FStrProperty>());
+
+		const FStrProperty* StringProperty = CastField<FStrProperty>(Property);
+
+		const void* ValueAddr = Property->ContainerPtrToValuePtr<void>(InCaptureMetadata);
+
+		FString Value = StringProperty->GetPropertyValue(ValueAddr);
+
+		FString OldValue;
+		if (Metadata.HasValue(InObject, Property->GetFName()))
+		{
+			OldValue = Metadata.GetValue(InObject, Property->GetFName());
+		}
+
+		if (Value != OldValue)
+		{
+			Metadata.SetValue(InObject, Property->GetFName(), *Value);
+
+			bIsEdited = true;
+		}
+		
+		++It;
+	}
+
+	if (bIsEdited)
+	{
+		InObject->MarkPackageDirty();
+	}
+}
+
+UCaptureMetadata* UCaptureMetadata::GetCaptureMetadata(const UObject* InObject)
+{
+	if (!IsValid(InObject))
+	{
+		return nullptr;
+	}
+
+	FMetaData& Metadata = InObject->GetPackage()->GetMetaData();
+
+	UCaptureMetadata* CaptureMetadata = NewObject<UCaptureMetadata>(GetTransientPackage(), InObject->GetFName());
+	CaptureMetadata->OwnerName = InObject->GetFName().ToString();
+
+	TMap<FName, FString>* MetadataMap = Metadata.GetMapForObject(InObject);
+	if (!MetadataMap)
+	{
+		return CaptureMetadata;
+	}
+
+	for (const TPair<FName, FString>& MetadataPair : *MetadataMap)
+	{
+		FProperty* Property = CaptureMetadata->GetClass()->FindPropertyByName(MetadataPair.Key);
+
+		if (Property)
+		{
+			check(Property->IsA<FStrProperty>());
+
+			const FStrProperty* StringProperty = CastField<FStrProperty>(Property);
+			void* ValueAddr = Property->ContainerPtrToValuePtr<void>(CaptureMetadata);
+
+			StringProperty->SetPropertyValue(ValueAddr, MetadataPair.Value);
+		}
+	}
+
+	return CaptureMetadata;
+}
+
+void UCaptureMetadata::ClearCaptureMetadata(const UObject* InObject)
+{
+	if (!IsValid(InObject))
+	{
+		return;
+	}
+
+	FMetaData& Metadata = InObject->GetPackage()->GetMetaData();
+
+	TMap<FName, FString>* MetadataMap = Metadata.GetMapForObject(InObject);
+	if (!MetadataMap)
+	{
+		return;
+	}
+
+	TMap<FName, FString> MetadataCopy = *MetadataMap;
+
+	bool bIsEdited = false;
+	for (const TPair<FName, FString>& MetadataPair : MetadataCopy)
+	{
+		FProperty* Property = UCaptureMetadata::StaticClass()->FindPropertyByName(MetadataPair.Key);
+
+		if (Property)
+		{
+			check(Property->IsA<FStrProperty>());
+
+			if (Metadata.HasValue(InObject, MetadataPair.Key))
+			{
+				Metadata.RemoveValue(InObject, MetadataPair.Key);
+
+				bIsEdited = true;
+			}
+		}
+	}
+
+	if (bIsEdited)
+	{
+		InObject->MarkPackageDirty();
+	}
+}
+
+bool UCaptureMetadata::ShowCaptureMetadataObjects(const FText& InTitle, const TArray<UObject*>& InObjects, const FCaptureMetadataWindowOptions& InOptions)
+{
+#if WITH_EDITOR
+	for (UObject* Object : InObjects)
+	{
+		UCaptureMetadata* CaptureMetadata = StaticCast<UCaptureMetadata*>(Object);
+		CaptureMetadata->bIsEditable = InOptions.bAllowEdit;
+	}
+
+	FEditorDialogLibraryObjectDetailsViewOptions Options;
+	Options.bShowObjectName = false;
+	Options.bAllowResizing = true;
+	Options.MinWidth = 400;
+	Options.MinHeight = 200;
+
+	return UEditorDialogLibrary::ShowObjectsDetailsView(InTitle, InObjects, Options);
+#else
+	return false;
+#endif
+}
+
+bool UCaptureMetadata::IsEditable() const
+{
+	return bIsEditable;
+}
+
+FString UCaptureMetadata::GetOwnerName() const
+{
+	return OwnerName;
+}
+
+void UCaptureMetadata::PostInitProperties()
+{
+	Super::PostInitProperties();
+
+	// This function ensures that all properties defined in the UCaptureMetadata class are of type FString
+	TFieldIterator<FProperty> It(GetClass());
+
+	while (It)
+	{
+		FProperty* Property = *It;
+		check(Property->IsA<FStrProperty>());
+		++It;
+	}
+}
\ No newline at end of file
diff --git a/Engine/Plugins/VirtualProduction/CaptureData/Source/CaptureDataEditor/Private/Customizations/CaptureMetadataCustomization.cpp b/Engine/Plugins/VirtualProduction/CaptureData/Source/CaptureDataEditor/Private/Customizations/CaptureMetadataCustomization.cpp
new file mode 100644
--- /dev/null
+++ b/Engine/Plugins/VirtualProduction/CaptureData/Source/CaptureDataEditor/Private/Customizations/CaptureMetadataCustomization.cpp
@@ -0,0 +1,53 @@
+// Copyright Epic Games, Inc.All Rights Reserved.
+
+#include "Customizations/CaptureMetadataCustomization.h"
+
+#include "DetailLayoutBuilder.h"
+#include "DetailCategoryBuilder.h"
+
+#include "CaptureMetadata.h"
+
+TSharedRef<IDetailCustomization> FCaptureMetadataCustomization::MakeInstance()
+{
+	return MakeShared<FCaptureMetadataCustomization>();
+}
+
+void FCaptureMetadataCustomization::CustomizeDetails(IDetailLayoutBuilder& InDetailBuilder)
+{
+	TArray<TWeakObjectPtr<UObject>> Objects;
+	InDetailBuilder.GetObjectsBeingCustomized(Objects);
+
+	check(!Objects.IsEmpty());
+
+	TWeakObjectPtr<UCaptureMetadata> CaptureMetadataPtr = StaticCast<UCaptureMetadata*>(Objects[0].Get());
+
+	// We are setting the Category name for all properties to group them under a owner object name
+	// 
+	// When showing more then one Capture Metadata, Category Name can be used to search for the owner object
+	const FString NewCategoryName = CaptureMetadataPtr->GetOwnerName();
+	TArray<FName> CategoryNames;
+	InDetailBuilder.GetCategoryNames(CategoryNames);
+
+	for (const FName& Category : CategoryNames)
+	{
+		IDetailCategoryBuilder& CategoryBuilder = InDetailBuilder.EditCategory(Category);
+		CategoryBuilder.SetDisplayName(FText::FromString(NewCategoryName));
+	}
+
+	if (!CaptureMetadataPtr->IsEditable())
+	{
+		TFieldIterator<FProperty> Iterator(InDetailBuilder.GetBaseClass());
+		while (Iterator)
+		{
+			FProperty* Property = *Iterator;
+			check(Property->IsA<FStrProperty>());
+
+			TSharedRef<IPropertyHandle> PropertyHandle = InDetailBuilder.GetProperty(Property->GetFName());
+			IDetailPropertyRow* PropertyRow = InDetailBuilder.EditDefaultProperty(PropertyHandle);
+
+			PropertyRow->IsEnabled(false);
+
+			++Iterator;
+		}
+	}
+}
\ No newline at end of file
diff --git a/Engine/Plugins/VirtualProduction/CaptureData/Source/CaptureDataEditor/Private/Customizations/CaptureMetadataCustomization.h b/Engine/Plugins/VirtualProduction/CaptureData/Source/CaptureDataEditor/Private/Customizations/CaptureMetadataCustomization.h
new file mode 100644
--- /dev/null
+++ b/Engine/Plugins/VirtualProduction/CaptureData/Source/CaptureDataEditor/Private/Customizations/CaptureMetadataCustomization.h
@@ -0,0 +1,18 @@
+// Copyright Epic Games, Inc.All Rights Reserved.
+
+#pragma once
+
+#include "CoreMinimal.h"
+#include "UObject/Object.h"
+#include "IDetailCustomization.h"
+
+class FCaptureMetadataCustomization
+	: public IDetailCustomization
+{
+public:
+	static TSharedRef<IDetailCustomization> MakeInstance();
+
+	//~Begin IDetailCustomization interface
+	virtual void CustomizeDetails(IDetailLayoutBuilder& InDetailBuilder) override;
+	//~End IDetailCustomization interface
+};
diff --git a/Engine/Plugins/VirtualProduction/CaptureData/Source/CaptureDataEditor/Public/CaptureMetadata.h b/Engine/Plugins/VirtualProduction/CaptureData/Source/CaptureDataEditor/Public/CaptureMetadata.h
new file mode 100644
--- /dev/null
+++ b/Engine/Plugins/VirtualProduction/CaptureData/Source/CaptureDataEditor/Public/CaptureMetadata.h
@@ -0,0 +1,54 @@
+// Copyright Epic Games, Inc. All Rights Reserved.
+
+#pragma once
+
+#include "UObject/Object.h"
+
+#include "CaptureMetadata.generated.h"
+
+#define UE_API CAPTUREDATAEDITOR_API
+
+USTRUCT(MinimalAPI, BlueprintType, Blueprintable)
+struct FCaptureMetadataWindowOptions
+{
+	GENERATED_USTRUCT_BODY()
+
+	UPROPERTY(EditAnywhere, BlueprintReadWrite, Category = "Options")
+	bool bAllowEdit = true;
+};
+
+UCLASS(MinimalAPI, BlueprintType, Blueprintable)
+class UCaptureMetadata
+	: public UObject
+{
+	GENERATED_BODY()
+
+public:
+
+	UFUNCTION(BlueprintCallable, Category = "Metadata")
+	static UE_API void SetCaptureMetadata(UObject* InObject, const UCaptureMetadata* InCaptureMetadata);
+
+	UFUNCTION(BlueprintCallable, Category = "Metadata")
+	static UE_API UCaptureMetadata* GetCaptureMetadata(const UObject* InObject);
+
+	UFUNCTION(BlueprintCallable, Category = "Metadata")
+	static UE_API void ClearCaptureMetadata(const UObject* InObject);
+
+	UFUNCTION(BlueprintCallable, Category = "Metadata")
+	static UE_API bool ShowCaptureMetadataObjects(const FText& InTitle, const TArray<UObject*>& InObjects, const FCaptureMetadataWindowOptions& InOptions);
+
+	UPROPERTY(EditAnywhere, BlueprintReadWrite, meta = (ExposeOnSpawn), Category = "Metadata")
+	FString CameraId;
+
+	bool IsEditable() const;
+	FString GetOwnerName() const;
+
+private:
+
+	bool bIsEditable = true;
+	FString OwnerName;
+
+	virtual void PostInitProperties() override;
+};
+
+#undef UE_API
\ No newline at end of file
