Changing capture metadata property category

#Synced-CL 49269982

diff --git a/Engine/Plugins/VirtualProduction/CaptureData/Source/CaptureDataEditor/Public/CaptureMetadata.h b/Engine/Plugins/VirtualProduction/CaptureData/Source/CaptureDataEditor/Public/CaptureMetadata.h
--- a/Engine/Plugins/VirtualProduction/CaptureData/Source/CaptureDataEditor/Public/CaptureMetadata.h
+++ b/Engine/Plugins/VirtualProduction/CaptureData/Source/CaptureDataEditor/Public/CaptureMetadata.h
@@ -13,7 +13,7 @@ struct FCaptureMetadataWindowOptions
 {
 	GENERATED_USTRUCT_BODY()
 
-	UPROPERTY(EditAnywhere, BlueprintReadWrite, Category = "Options")
+	UPROPERTY(EditAnywhere, BlueprintReadWrite, Category = "Capture Metadata | Options")
 	bool bAllowEdit = true;
 };
 
@@ -25,19 +25,19 @@ class UCaptureMetadata
 
 public:
 
-	UFUNCTION(BlueprintCallable, Category = "Metadata")
+	UFUNCTION(BlueprintCallable, Category = "Capture Metadata")
 	static UE_API void SetCaptureMetadata(UObject* InObject, const UCaptureMetadata* InCaptureMetadata);
 
-	UFUNCTION(BlueprintCallable, Category = "Metadata")
+	UFUNCTION(BlueprintCallable, Category = "Capture Metadata")
 	static UE_API UCaptureMetadata* GetCaptureMetadata(const UObject* InObject);
 
-	UFUNCTION(BlueprintCallable, Category = "Metadata")
+	UFUNCTION(BlueprintCallable, Category = "Capture Metadata")
 	static UE_API void ClearCaptureMetadata(const UObject* InObject);
 
-	UFUNCTION(BlueprintCallable, Category = "Metadata")
+	UFUNCTION(BlueprintCallable, Category = "Capture Metadata")
 	static UE_API bool ShowCaptureMetadataObjects(const FText& InTitle, const TArray<UObject*>& InObjects, const FCaptureMetadataWindowOptions& InOptions);
 
-	UPROPERTY(EditAnywhere, BlueprintReadWrite, meta = (ExposeOnSpawn), Category = "Metadata")
+	UPROPERTY(EditAnywhere, BlueprintReadWrite, meta = (ExposeOnSpawn), Category = "Capture Metadata")
 	FString CameraId;
 
 	bool IsEditable() const;
