Creating Metadata object without setting name

#Synced-CL 52201630

diff --git a/Engine/Plugins/VirtualProduction/CaptureData/Source/CaptureDataEditor/Private/CaptureMetadata.cpp b/Engine/Plugins/VirtualProduction/CaptureData/Source/CaptureDataEditor/Private/CaptureMetadata.cpp
--- a/Engine/Plugins/VirtualProduction/CaptureData/Source/CaptureDataEditor/Private/CaptureMetadata.cpp
+++ b/Engine/Plugins/VirtualProduction/CaptureData/Source/CaptureDataEditor/Private/CaptureMetadata.cpp
@@ -14,13 +14,7 @@ void UCaptureMetadata::SetCaptureMetadata(UObject* InObject, const UCaptureMetad
 
 UCaptureMetadata* UCaptureMetadata::GetCaptureMetadata(const UObject* InObject)
 {
-	UCaptureMetadata* Metadata = UE::GetMetadataObject<UCaptureMetadata>(InObject);
-	if (Metadata)
-	{
-		Metadata->OwnerName = InObject->GetName();
-	}
-
-	return Metadata;
+	return UE::GetMetadataObject<UCaptureMetadata>(InObject);
 }
 
 void UCaptureMetadata::ClearCaptureMetadata(const UObject* InObject)
@@ -48,5 +42,5 @@ bool UCaptureMetadata::IsEditable() const
 
 FString UCaptureMetadata::GetOwnerName() const
 {
-	return OwnerName;
+	return OwnerName.ToString();
 }
diff --git a/Engine/Plugins/VirtualProduction/CaptureData/Source/CaptureDataEditor/Private/Customizations/CaptureMetadataCustomization.cpp b/Engine/Plugins/VirtualProduction/CaptureData/Source/CaptureDataEditor/Private/Customizations/CaptureMetadataCustomization.cpp
--- a/Engine/Plugins/VirtualProduction/CaptureData/Source/CaptureDataEditor/Private/Customizations/CaptureMetadataCustomization.cpp
+++ b/Engine/Plugins/VirtualProduction/CaptureData/Source/CaptureDataEditor/Private/Customizations/CaptureMetadataCustomization.cpp
@@ -24,7 +24,7 @@ void FCaptureMetadataCustomization::CustomizeDetails(IDetailLayoutBuilder& InDet
 	// We are setting the Category name for all properties to group them under a owner object name
 	// 
 	// When showing more then one Capture Metadata, Category Name can be used to search for the owner object
-	const FString NewCategoryName = CaptureMetadataPtr->GetOwnerName();
+	const FString NewCategoryName = CaptureMetadataPtr->OwnerName.ToString();
 	TArray<FName> CategoryNames;
 	InDetailBuilder.GetCategoryNames(CategoryNames);
 
@@ -34,18 +34,22 @@ void FCaptureMetadataCustomization::CustomizeDetails(IDetailLayoutBuilder& InDet
 		CategoryBuilder.SetDisplayName(FText::FromString(NewCategoryName));
 	}
 
+	InDetailBuilder.HideProperty(InDetailBuilder.GetProperty(GET_MEMBER_NAME_CHECKED(UBaseCaptureMetadata, OwnerName), UBaseCaptureMetadata::StaticClass()));
+
 	if (!CaptureMetadataPtr->IsEditable())
 	{
 		TFieldIterator<FProperty> Iterator(InDetailBuilder.GetBaseClass());
 		while (Iterator)
 		{
 			FProperty* Property = *Iterator;
-			check(Property->IsA<FStrProperty>());
 
 			TSharedRef<IPropertyHandle> PropertyHandle = InDetailBuilder.GetProperty(Property->GetFName());
 			IDetailPropertyRow* PropertyRow = InDetailBuilder.EditDefaultProperty(PropertyHandle);
 
-			PropertyRow->IsEnabled(false);
+			if (PropertyRow)
+			{
+				PropertyRow->IsEnabled(false);
+			}
 
 			++Iterator;
 		}
diff --git a/Engine/Plugins/VirtualProduction/CaptureData/Source/CaptureDataEditor/Public/CameraCalibrationMetadata.h b/Engine/Plugins/VirtualProduction/CaptureData/Source/CaptureDataEditor/Public/CameraCalibrationMetadata.h
--- a/Engine/Plugins/VirtualProduction/CaptureData/Source/CaptureDataEditor/Public/CameraCalibrationMetadata.h
+++ b/Engine/Plugins/VirtualProduction/CaptureData/Source/CaptureDataEditor/Public/CameraCalibrationMetadata.h
@@ -3,6 +3,7 @@
 #pragma once
 
 #include "UObject/Object.h"
+#include "Metadata/MetadataHandler.h"
 
 #include "Misc/Timecode.h"
 
@@ -11,7 +12,7 @@
 #define UE_API CAPTUREDATAEDITOR_API
 
 UCLASS(MinimalAPI, BlueprintType, Blueprintable)
-class UCameraCalibrationMetadata : public UObject
+class UCameraCalibrationMetadata : public UBaseCaptureMetadata
 {
 public:
 
diff --git a/Engine/Plugins/VirtualProduction/CaptureData/Source/CaptureDataEditor/Public/CaptureMetadata.h b/Engine/Plugins/VirtualProduction/CaptureData/Source/CaptureDataEditor/Public/CaptureMetadata.h
--- a/Engine/Plugins/VirtualProduction/CaptureData/Source/CaptureDataEditor/Public/CaptureMetadata.h
+++ b/Engine/Plugins/VirtualProduction/CaptureData/Source/CaptureDataEditor/Public/CaptureMetadata.h
@@ -3,6 +3,7 @@
 #pragma once
 
 #include "UObject/Object.h"
+#include "Metadata/MetadataHandler.h"
 
 #include "CaptureMetadata.generated.h"
 
@@ -19,7 +20,7 @@ struct FCaptureMetadataWindowOptions
 
 UCLASS(MinimalAPI, BlueprintType, Blueprintable)
 class UCaptureMetadata
-	: public UObject
+	: public UBaseCaptureMetadata
 {
 	GENERATED_BODY()
 
@@ -46,7 +47,6 @@ public:
 private:
 
 	bool bIsEditable = true;
-	FString OwnerName;
 };
 
 #undef UE_API
\ No newline at end of file
diff --git a/Engine/Plugins/VirtualProduction/CaptureData/Source/CaptureDataEditor/Public/Metadata/MetadataHandler.h b/Engine/Plugins/VirtualProduction/CaptureData/Source/CaptureDataEditor/Public/Metadata/MetadataHandler.h
--- a/Engine/Plugins/VirtualProduction/CaptureData/Source/CaptureDataEditor/Public/Metadata/MetadataHandler.h
+++ b/Engine/Plugins/VirtualProduction/CaptureData/Source/CaptureDataEditor/Public/Metadata/MetadataHandler.h
@@ -7,11 +7,24 @@
 
 #include "Concepts/DerivedFrom.h"
 
+#include "MetadataHandler.generated.h"
+
+UCLASS(MinimalAPI, BlueprintType, Blueprintable)
+class UBaseCaptureMetadata : public UObject
+{
+	GENERATED_BODY()
+
+public:
+
+	UPROPERTY(VisibleAnywhere, BlueprintReadOnly, Category="Base Capture Metadata")
+	FName OwnerName;
+};
+
 namespace UE
 {
 
 template<typename T>
-	requires UE::CDerivedFrom<T, UObject>
+	requires UE::CDerivedFrom<T, UBaseCaptureMetadata>
 void SetMetadataObject(UObject* InObject, T* InMetadata)
 {
 	if (!IsValid(InObject))
@@ -35,22 +48,25 @@ void SetMetadataObject(UObject* InObject, T* InMetadata)
 	{
 		FProperty* Property = *It;
 
-		const void* ValueAddr = Property->ContainerPtrToValuePtr<void>(InMetadata);
+		if (Property->GetFName() != GET_MEMBER_NAME_CHECKED(UBaseCaptureMetadata, OwnerName))
+		{
+			const void* ValueAddr = Property->ContainerPtrToValuePtr<void>(InMetadata);
 
-		FString Value;
-		Property->ExportText_Direct(Value, ValueAddr, nullptr, InMetadata, 0);
+			FString Value;
+			Property->ExportText_Direct(Value, ValueAddr, nullptr, InMetadata, 0);
 
-		FString OldValue;
-		if (Metadata.HasValue(InObject, Property->GetFName()))
-		{
-			OldValue = Metadata.GetValue(InObject, Property->GetFName());
-		}
+			FString OldValue;
+			if (Metadata.HasValue(InObject, Property->GetFName()))
+			{
+				OldValue = Metadata.GetValue(InObject, Property->GetFName());
+			}
 
-		if (Value != OldValue)
-		{
-			Metadata.SetValue(InObject, Property->GetFName(), *Value);
+			if (Value != OldValue)
+			{
+				Metadata.SetValue(InObject, Property->GetFName(), *Value);
 
-			bIsEdited = true;
+				bIsEdited = true;
+			}
 		}
 
 		++It;
@@ -64,7 +80,7 @@ void SetMetadataObject(UObject* InObject, T* InMetadata)
 }
 
 template<typename T>
-	requires UE::CDerivedFrom<T, UObject>
+	requires UE::CDerivedFrom<T, UBaseCaptureMetadata>
 T* GetMetadataObject(const UObject* InObject)
 {
 	if (!IsValid(InObject))
@@ -81,10 +97,16 @@ T* GetMetadataObject(const UObject* InObject)
 		return nullptr;
 	}
 
-	T* MetadataObject = NewObject<T>(GetTransientPackage(), InObject->GetFName());
+	T* MetadataObject = NewObject<T>(GetTransientPackage());
+	MetadataObject->OwnerName = InObject->GetFName();
 
 	for (const TPair<FName, FString>& MetadataPair : *MetadataMap)
 	{
+		if (MetadataPair.Key == GET_MEMBER_NAME_CHECKED(UBaseCaptureMetadata, OwnerName))
+		{
+			continue;
+		}
+
 		FProperty* Property = MetadataObject->GetClass()->FindPropertyByName(MetadataPair.Key);
 
 		if (Property)
@@ -101,7 +123,7 @@ T* GetMetadataObject(const UObject* InObject)
 }
 
 template<typename T>
-	requires UE::CDerivedFrom<T, UObject>
+	requires UE::CDerivedFrom<T, UBaseCaptureMetadata>
 void ClearMetadataObject(const UObject* InObject)
 {
 	if (!IsValid(InObject))
