Added log messages to ImageSequenceUtils

#Synced-CL 52679542

diff --git a/Engine/Plugins/VirtualProduction/CaptureData/Source/CaptureDataCore/Private/ImageSequenceUtils.cpp b/Engine/Plugins/VirtualProduction/CaptureData/Source/CaptureDataCore/Private/ImageSequenceUtils.cpp
--- a/Engine/Plugins/VirtualProduction/CaptureData/Source/CaptureDataCore/Private/ImageSequenceUtils.cpp
+++ b/Engine/Plugins/VirtualProduction/CaptureData/Source/CaptureDataCore/Private/ImageSequenceUtils.cpp
@@ -11,6 +11,7 @@
 #include "Modules/ModuleManager.h"
 #include "ProfilingDebugging/ScopedTimers.h"
 #include "TakeDirectoryUtils.h"
+#include "CaptureDataLog.h"
 
 
 bool FImageSequenceUtils::GetImageSequencePathAndFilesFromAsset(const UImgMediaSource* InImgSequence, FString& OutFullSequencePath, TArray<FString>& OutImageFiles)
@@ -72,6 +73,7 @@ bool FImageSequenceUtils::GetImageSequenceInfoFromAsset(const class UImgMediaSou
 {
 	if (InImgSequence == nullptr)
 	{
+		UE_LOGF(LogCaptureDataCore, Warning, "GetImageSequenceInfoFromAsset: InImgSequence is null.");
 		return false;
 	}
 
@@ -85,29 +87,40 @@ bool FImageSequenceUtils::GetImageSequenceInfoFromPath(const FString& InFullSequ
 
 	if (!bFoundImages)
 	{
+		UE_LOGF(LogCaptureDataCore, Warning, "GetImageSequenceInfoFromPath: No image files found in '%ls'.", *InFullSequencePath);
 		return false;
 	}
 
 	OutNumImages = ImageFiles.Num();
 
+	if (ImageFiles.IsEmpty())
+	{
+		UE_LOGF(LogCaptureDataCore, Warning, "GetImageSequenceInfoFromPath: Image file list is empty for '%ls'.", *InFullSequencePath);
+		return false;
+	}
+
 	const FString SampleImagePath = InFullSequencePath / ImageFiles[0];
 
 	TArray<uint8> RawFileData;
-	if (FFileHelper::LoadFileToArray(RawFileData, *SampleImagePath))
+	if (!FFileHelper::LoadFileToArray(RawFileData, *SampleImagePath))
 	{
-		IImageWrapperModule& ImageWrapperModule = FModuleManager::LoadModuleChecked<IImageWrapperModule>(TEXT("ImageWrapper"));
-		EImageFormat ImageFormat = ImageWrapperModule.DetectImageFormat(RawFileData.GetData(), RawFileData.Num());
-		TSharedPtr<IImageWrapper> ImageWrapper = ImageWrapperModule.CreateImageWrapper(ImageFormat);
+		UE_LOGF(LogCaptureDataCore, Warning, "GetImageSequenceInfoFromPath: Failed to load file '%ls'.", *SampleImagePath);
+		return false;
+	}
 
-		if (ImageWrapper.IsValid() && ImageWrapper->SetCompressed(RawFileData.GetData(), RawFileData.Num()))
-		{
-			OutDimensions.X = ImageWrapper->GetWidth();
-			OutDimensions.Y = ImageWrapper->GetHeight();
-			return true;
-		}
+	IImageWrapperModule& ImageWrapperModule = FModuleManager::LoadModuleChecked<IImageWrapperModule>(TEXT("ImageWrapper"));
+	EImageFormat ImageFormat = ImageWrapperModule.DetectImageFormat(RawFileData.GetData(), RawFileData.Num());
+	TSharedPtr<IImageWrapper> ImageWrapper = ImageWrapperModule.CreateImageWrapper(ImageFormat);
+
+	if (!ImageWrapper.IsValid() || !ImageWrapper->SetCompressed(RawFileData.GetData(), RawFileData.Num()))
+	{
+		UE_LOGF(LogCaptureDataCore, Warning, "GetImageSequenceInfoFromPath: Failed to decode image '%ls'.", *SampleImagePath);
+		return false;
 	}
 
-	return false;
+	OutDimensions.X = ImageWrapper->GetWidth();
+	OutDimensions.Y = ImageWrapper->GetHeight();
+	return true;
 }
 
 // We consider that the image sequence is valid if it has at least one image
